Change Meshdescription to use FName instead of indexes.

diff --git a/Engine/Plugins/Editor/GLTFImporter/Source/GLTFImporter/Private/GLTFStaticMesh.cpp b/Engine/Plugins/Editor/GLTFImporter/Source/GLTFImporter/Private/GLTFStaticMesh.cpp
--- a/Engine/Plugins/Editor/GLTFImporter/Source/GLTFImporter/Private/GLTFStaticMesh.cpp
+++ b/Engine/Plugins/Editor/GLTFImporter/Source/GLTFImporter/Private/GLTFStaticMesh.cpp
@@ -166,8 +166,7 @@ UStaticMesh* ImportStaticMesh(const FAsset& Asset, const TArray<UMaterial*>& Mat
 
 	UStaticMesh* StaticMesh = NewObject<UStaticMesh>(AssetPackage, FName(*AssetName), Flags);
 
-	StaticMesh->SourceModels.AddDefaulted();
-	FStaticMeshSourceModel& SourceModel = StaticMesh->SourceModels.Last();
+	FStaticMeshSourceModel& SourceModel = StaticMesh->AddSourceModel();
 	FMeshBuildSettings& Settings = SourceModel.BuildSettings;
 
 	Settings.bRecomputeNormals = false;
diff --git a/Engine/Plugins/Editor/SpeedTreeImporter/Source/SpeedTreeImporter/Private/SpeedTreeImportFactory.cpp b/Engine/Plugins/Editor/SpeedTreeImporter/Source/SpeedTreeImporter/Private/SpeedTreeImportFactory.cpp
--- a/Engine/Plugins/Editor/SpeedTreeImporter/Source/SpeedTreeImporter/Private/SpeedTreeImportFactory.cpp
+++ b/Engine/Plugins/Editor/SpeedTreeImporter/Source/SpeedTreeImporter/Private/SpeedTreeImportFactory.cpp
@@ -1783,15 +1783,15 @@ UObject* USpeedTreeImportFactory::FactoryCreateBinary7(UClass* InClass, UObject*
 							}							
 						}
 
-						FStaticMeshSourceModel* LODModel = new (StaticMesh->SourceModels) FStaticMeshSourceModel();
-						LODModel->BuildSettings.bRecomputeNormals = false;
-						LODModel->BuildSettings.bRecomputeTangents = false;
-						LODModel->BuildSettings.bRemoveDegenerates = true;
-						LODModel->BuildSettings.bUseHighPrecisionTangentBasis = false;
-						LODModel->BuildSettings.bUseFullPrecisionUVs = false;
-						LODModel->BuildSettings.bGenerateLightmapUVs = false;
-						LODModel->ScreenSize = 0.1f / FMath::Pow(2.0f, StaticMesh->SourceModels.Num() - 1);
-						LODModel->SaveRawMesh(RawMesh);
+						FStaticMeshSourceModel& LODModel = StaticMesh->AddSourceModel();
+						LODModel.BuildSettings.bRecomputeNormals = false;
+						LODModel.BuildSettings.bRecomputeTangents = false;
+						LODModel.BuildSettings.bRemoveDegenerates = true;
+						LODModel.BuildSettings.bUseHighPrecisionTangentBasis = false;
+						LODModel.BuildSettings.bUseFullPrecisionUVs = false;
+						LODModel.BuildSettings.bGenerateLightmapUVs = false;
+						LODModel.ScreenSize = 0.1f / FMath::Pow(2.0f, StaticMesh->SourceModels.Num() - 1);
+						LODModel.SaveRawMesh(RawMesh);
 
 						for (int32 MaterialIndex = 0; MaterialIndex < StaticMesh->StaticMaterials.Num(); ++MaterialIndex)
 						{
@@ -1897,15 +1897,15 @@ UObject* USpeedTreeImportFactory::FactoryCreateBinary7(UClass* InClass, UObject*
 						}
 					}
 
-					FStaticMeshSourceModel* LODModel = new (StaticMesh->SourceModels) FStaticMeshSourceModel();
-					LODModel->BuildSettings.bRecomputeNormals = false;
-					LODModel->BuildSettings.bRecomputeTangents = false;
-					LODModel->BuildSettings.bRemoveDegenerates = true;
-					LODModel->BuildSettings.bUseHighPrecisionTangentBasis = false;
-					LODModel->BuildSettings.bUseFullPrecisionUVs = false;
-					LODModel->BuildSettings.bGenerateLightmapUVs = false;
-					LODModel->ScreenSize = 0.1f / FMath::Pow(2.0f, StaticMesh->SourceModels.Num() - 1);
-					LODModel->SaveRawMesh(RawMesh);
+					FStaticMeshSourceModel& LODModel = StaticMesh->AddSourceModel();
+					LODModel.BuildSettings.bRecomputeNormals = false;
+					LODModel.BuildSettings.bRecomputeTangents = false;
+					LODModel.BuildSettings.bRemoveDegenerates = true;
+					LODModel.BuildSettings.bUseHighPrecisionTangentBasis = false;
+					LODModel.BuildSettings.bUseFullPrecisionUVs = false;
+					LODModel.BuildSettings.bGenerateLightmapUVs = false;
+					LODModel.ScreenSize = 0.1f / FMath::Pow(2.0f, StaticMesh->SourceModels.Num() - 1);
+					LODModel.SaveRawMesh(RawMesh);
 					// Add mesh section info entry for billboard LOD (only one section/material index)
 					const int32 LODIndex = StaticMesh->SourceModels.Num() - 1;
 					FMeshSectionInfo Info = StaticMesh->SectionInfoMap.Get(LODIndex, 0);
@@ -2049,23 +2049,23 @@ UObject* USpeedTreeImportFactory::FactoryCreateBinary8(UClass* InClass, UObject*
 
 			for (uint32 LODIndex = 0; LODIndex < SpeedTree.Lods().Count(); ++LODIndex)
 			{
-				FStaticMeshSourceModel* LODModel = new (StaticMesh->SourceModels) FStaticMeshSourceModel();
-				LODModel->BuildSettings.SrcLightmapIndex = 1;
-				LODModel->BuildSettings.DstLightmapIndex = 1;
-				LODModel->BuildSettings.bRecomputeNormals = false;
-				LODModel->BuildSettings.bRecomputeTangents = false;
-				LODModel->BuildSettings.bRemoveDegenerates = true;
-				LODModel->BuildSettings.bUseHighPrecisionTangentBasis = false;
-				LODModel->BuildSettings.bUseFullPrecisionUVs = false;
-				LODModel->BuildSettings.bGenerateLightmapUVs = false;
+				FStaticMeshSourceModel& LODModel = StaticMesh->AddSourceModel();
+				LODModel.BuildSettings.SrcLightmapIndex = 1;
+				LODModel.BuildSettings.DstLightmapIndex = 1;
+				LODModel.BuildSettings.bRecomputeNormals = false;
+				LODModel.BuildSettings.bRecomputeTangents = false;
+				LODModel.BuildSettings.bRemoveDegenerates = true;
+				LODModel.BuildSettings.bUseHighPrecisionTangentBasis = false;
+				LODModel.BuildSettings.bUseFullPrecisionUVs = false;
+				LODModel.BuildSettings.bGenerateLightmapUVs = false;
 
 				if (Options->SpeedTreeImportData->LODType == ILT_IndividualActors)
 				{
-					LODModel->ScreenSize = FMath::Lerp(1.0f, 0.1f, FMath::Square(LODIndex * Denominator));
+					LODModel.ScreenSize = FMath::Lerp(1.0f, 0.1f, FMath::Square(LODIndex * Denominator));
 				}
 				else
 				{
-					LODModel->ScreenSize = FMath::Lerp(1.0f, 0.25f, LODIndex * Denominator);
+					LODModel.ScreenSize = FMath::Lerp(1.0f, 0.25f, LODIndex * Denominator);
 				}
 			}
 		}
diff --git a/Engine/Plugins/Editor/USDImporter/Source/USDImporter/Private/StaticMeshImporter.cpp b/Engine/Plugins/Editor/USDImporter/Source/USDImporter/Private/StaticMeshImporter.cpp
--- a/Engine/Plugins/Editor/USDImporter/Source/USDImporter/Private/StaticMeshImporter.cpp
+++ b/Engine/Plugins/Editor/USDImporter/Source/USDImporter/Private/StaticMeshImporter.cpp
@@ -300,9 +300,10 @@ UStaticMesh* FUSDStaticMeshImporter::ImportStaticMesh(FUsdImportContext& ImportC
 			if (!NewMesh->SourceModels.IsValidIndex(LODIndex))
 			{
 				// Add one LOD 
-				NewMesh->SourceModels.AddDefaulted();
+				NewMesh->AddSourceModel();
 			}
 
+			
 			FStaticMeshSourceModel& SrcModel = NewMesh->SourceModels[LODIndex];
 
 			RawTriangles.CompactMaterialIndices();
diff --git a/Engine/Plugins/Experimental/AlembicImporter/Source/AlembicLibrary/Private/AbcImporter.cpp b/Engine/Plugins/Experimental/AlembicImporter/Source/AlembicLibrary/Private/AbcImporter.cpp
--- a/Engine/Plugins/Experimental/AlembicImporter/Source/AlembicLibrary/Private/AbcImporter.cpp
+++ b/Engine/Plugins/Experimental/AlembicImporter/Source/AlembicLibrary/Private/AbcImporter.cpp
@@ -797,7 +797,7 @@ UStaticMesh* FAbcImporter::CreateStaticMeshFromRawMesh(UObject* InParent, const
 	if (StaticMesh)
 	{
 		// Add the first LOD, we only support one
-		new(StaticMesh->SourceModels) FStaticMeshSourceModel();
+		StaticMesh->AddSourceModel();
 
 		// Generate a new lighting GUID (so its unique)
 		StaticMesh->LightingGuid = FGuid::NewGuid();
diff --git a/Engine/Plugins/Runtime/ProceduralMeshComponent/Source/ProceduralMeshComponentEditor/Private/ProceduralMeshComponentDetails.cpp b/Engine/Plugins/Runtime/ProceduralMeshComponent/Source/ProceduralMeshComponentEditor/Private/ProceduralMeshComponentDetails.cpp
--- a/Engine/Plugins/Runtime/ProceduralMeshComponent/Source/ProceduralMeshComponentEditor/Private/ProceduralMeshComponentDetails.cpp
+++ b/Engine/Plugins/Runtime/ProceduralMeshComponent/Source/ProceduralMeshComponentEditor/Private/ProceduralMeshComponentDetails.cpp
@@ -183,16 +183,16 @@ FReply FProceduralMeshComponentDetails::ClickedOnConvertToStaticMesh()
 				StaticMesh->LightingGuid = FGuid::NewGuid();
 
 				// Add source to new StaticMesh
-				FStaticMeshSourceModel* SrcModel = new (StaticMesh->SourceModels) FStaticMeshSourceModel();
-				SrcModel->BuildSettings.bRecomputeNormals = false;
-				SrcModel->BuildSettings.bRecomputeTangents = false;
-				SrcModel->BuildSettings.bRemoveDegenerates = false;
-				SrcModel->BuildSettings.bUseHighPrecisionTangentBasis = false;
-				SrcModel->BuildSettings.bUseFullPrecisionUVs = false;
-				SrcModel->BuildSettings.bGenerateLightmapUVs = true;
-				SrcModel->BuildSettings.SrcLightmapIndex = 0;
-				SrcModel->BuildSettings.DstLightmapIndex = 1;
-				SrcModel->SaveRawMesh(RawMesh);
+				FStaticMeshSourceModel& SrcModel = StaticMesh->AddSourceModel();
+				SrcModel.BuildSettings.bRecomputeNormals = false;
+				SrcModel.BuildSettings.bRecomputeTangents = false;
+				SrcModel.BuildSettings.bRemoveDegenerates = false;
+				SrcModel.BuildSettings.bUseHighPrecisionTangentBasis = false;
+				SrcModel.BuildSettings.bUseFullPrecisionUVs = false;
+				SrcModel.BuildSettings.bGenerateLightmapUVs = true;
+				SrcModel.BuildSettings.SrcLightmapIndex = 0;
+				SrcModel.BuildSettings.DstLightmapIndex = 1;
+				SrcModel.SaveRawMesh(RawMesh);
 
 				// Copy materials to new mesh
 				for (UMaterialInterface* Material : MeshMaterials)
diff --git a/Engine/Plugins/Tests/EditorTests/Source/EditorTests/Private/UnrealEd/MeshDescriptionTests.cpp b/Engine/Plugins/Tests/EditorTests/Source/EditorTests/Private/UnrealEd/MeshDescriptionTests.cpp
--- a/Engine/Plugins/Tests/EditorTests/Source/EditorTests/Private/UnrealEd/MeshDescriptionTests.cpp
+++ b/Engine/Plugins/Tests/EditorTests/Source/EditorTests/Private/UnrealEd/MeshDescriptionTests.cpp
@@ -311,6 +311,41 @@ void MeshDescriptionNumberArrayCompare(const FString& ConversionName, const FStr
 	}
 }
 
+template<typename U, typename V>
+void MeshDescriptionNumberFNameCompare(const FString& ConversionName, const FString& AssetName, FAutomationTestExecutionInfo& ExecutionInfo, bool& bIsSame, const U& ElementIterator, const FString& VectorArrayName, const TPolygonGroupAttributeArray<FName>& ReferenceArray, const TPolygonGroupAttributeArray<FName>& ResultArray)
+{
+	if (ReferenceArray.Num() != ResultArray.Num())
+	{
+		ExecutionInfo.AddEvent(FAutomationEvent(EAutomationEventType::Error, FString::Printf(TEXT("The %s conversion %s is not lossless, %s count is different. %s count expected [%d] result [%d]"),
+			*AssetName,
+			*ConversionName,
+			*VectorArrayName,
+			*VectorArrayName,
+			ReferenceArray.Num(),
+			ResultArray.Num())));
+		bIsSame = false;
+	}
+	else
+	{
+		for (V ElementID : ElementIterator.GetElementIDs())
+		{
+			if (ReferenceArray[ElementID] != ResultArray[ElementID])
+			{
+				ExecutionInfo.AddEvent(FAutomationEvent(EAutomationEventType::Error, FString::Printf(TEXT("The %s conversion %s is not lossless, %s array is different. Array index [%d] expected %s [%s] result [%s]"),
+					*AssetName,
+					*ConversionName,
+					*VectorArrayName,
+					ElementID.GetValue(),
+					*VectorArrayName,
+					*ReferenceArray[ElementID].ToString(),
+					*ResultArray[ElementID].ToString())));
+				bIsSame = false;
+				break;
+			}
+		}
+	}
+}
+
 bool FMeshDescriptionTest::CompareMeshDescription(const FString& AssetName, FAutomationTestExecutionInfo& ExecutionInfo, const UMeshDescription* ReferenceMeshDescription, const UMeshDescription* MeshDescription) const
 {
 	//////////////////////////////////////////////////////////////////////////
@@ -325,7 +360,7 @@ bool FMeshDescriptionTest::CompareMeshDescription(const FString& AssetName, FAut
 
 	const TEdgeAttributeArray<bool>& ReferenceEdgeHardnesses = ReferenceMeshDescription->EdgeAttributes().GetAttributes<bool>(MeshAttribute::Edge::IsHard);
 
-	const TPolygonGroupAttributeArray<int>& ReferencePolygonGroupMaterialIndex = ReferenceMeshDescription->PolygonGroupAttributes().GetAttributes<int>(MeshAttribute::PolygonGroup::MaterialIndex);
+	const TPolygonGroupAttributeArray<FName>& ReferencePolygonGroupMaterialName = ReferenceMeshDescription->PolygonGroupAttributes().GetAttributes<FName>(MeshAttribute::PolygonGroup::ImportedMaterialSlotName);
 
 
 	//////////////////////////////////////////////////////////////////////////
@@ -340,7 +375,7 @@ bool FMeshDescriptionTest::CompareMeshDescription(const FString& AssetName, FAut
 
 	const TEdgeAttributeArray<bool>& ResultEdgeHardnesses = MeshDescription->EdgeAttributes().GetAttributes<bool>(MeshAttribute::Edge::IsHard);
 
-	const TPolygonGroupAttributeArray<int>& ResultPolygonGroupMaterialIndex = MeshDescription->PolygonGroupAttributes().GetAttributes<int>(MeshAttribute::PolygonGroup::MaterialIndex);
+	const TPolygonGroupAttributeArray<FName>& ResultPolygonGroupMaterialName = MeshDescription->PolygonGroupAttributes().GetAttributes<FName>(MeshAttribute::PolygonGroup::ImportedMaterialSlotName);
 
 	
 	//////////////////////////////////////////////////////////////////////////
@@ -418,7 +453,7 @@ bool FMeshDescriptionTest::CompareMeshDescription(const FString& AssetName, FAut
 	}
 
 	//Polygon group ID
-	MeshDescriptionNumberArrayCompare<TPolygonGroupAttributeArray<int>, FPolygonGroupArray, FPolygonGroupID>(ConversionName, AssetName, ExecutionInfo, bAllSame, ReferenceMeshDescription->PolygonGroups(), TEXT("vertex instance binormals"), ReferencePolygonGroupMaterialIndex, ResultPolygonGroupMaterialIndex);
+	MeshDescriptionNumberFNameCompare<FPolygonGroupArray, FPolygonGroupID>(ConversionName, AssetName, ExecutionInfo, bAllSame, ReferenceMeshDescription->PolygonGroups(), TEXT("PolygonGroup Material Name"), ReferencePolygonGroupMaterialName, ResultPolygonGroupMaterialName);
 
 	return bAllSame;
 }
@@ -453,6 +488,14 @@ bool FMeshDescriptionTest::ConversionTest(FAutomationTestExecutionInfo& Executio
 
 		if (AssetMesh != nullptr)
 		{
+			TMap<FName, int32> MaterialMap;
+			TMap<int32, FName> MaterialMapInverse;
+			for (int32 MaterialIndex = 0; MaterialIndex < AssetMesh->StaticMaterials.Num(); ++MaterialIndex)
+			{
+				MaterialMap.Add(AssetMesh->StaticMaterials[MaterialIndex].ImportedMaterialSlotName, MaterialIndex);
+				MaterialMapInverse.Add(MaterialIndex, AssetMesh->StaticMaterials[MaterialIndex].ImportedMaterialSlotName);
+			}
+
 			//MeshDescription to RawMesh to MeshDescription
 			for(int32 LodIndex = 0; LodIndex < AssetMesh->SourceModels.Num(); ++LodIndex)
 			{
@@ -466,9 +509,9 @@ bool FMeshDescriptionTest::ConversionTest(FAutomationTestExecutionInfo& Executio
 				UMeshDescription* ResultAssetMesh = DuplicateObject<UMeshDescription>(ReferenceAssetMesh, GetTransientPackage(), NAME_None);
 				//Convert MeshDescription to FRawMesh
 				FRawMesh RawMesh;
-				FMeshDescriptionOperations::ConverToRawMesh(ResultAssetMesh, RawMesh);
+				FMeshDescriptionOperations::ConverToRawMesh(ResultAssetMesh, RawMesh, MaterialMap);
 				//Convert back the FRawmesh
-				FMeshDescriptionOperations::ConverFromRawMesh(RawMesh, ResultAssetMesh);
+				FMeshDescriptionOperations::ConverFromRawMesh(RawMesh, ResultAssetMesh, MaterialMapInverse);
 				if (!CompareMeshDescription(AssetName, ExecutionInfo, ReferenceAssetMesh, ResultAssetMesh))
 				{
 					bAllSame = false;
@@ -492,9 +535,9 @@ bool FMeshDescriptionTest::ConversionTest(FAutomationTestExecutionInfo& Executio
 				AssetMesh->SourceModels[LodIndex].LoadRawMesh(ResultRawMesh);
 				//Create a temporary Mesh Description
 				UMeshDescription* MeshDescription = NewObject<UMeshDescription>(GetTransientPackage(), NAME_None, RF_Standalone);
-				FMeshDescriptionOperations::ConverFromRawMesh(ResultRawMesh, MeshDescription);
+				FMeshDescriptionOperations::ConverFromRawMesh(ResultRawMesh, MeshDescription, MaterialMapInverse);
 				//Convert back the FRawmesh
-				FMeshDescriptionOperations::ConverToRawMesh(MeshDescription, ResultRawMesh);
+				FMeshDescriptionOperations::ConverToRawMesh(MeshDescription, ResultRawMesh, MaterialMap);
 				if (!CompareRawMesh(AssetName, ExecutionInfo, ReferenceRawMesh, ResultRawMesh))
 				{
 					bAllSame = false;
diff --git a/Engine/Source/Developer/AssetTools/Private/AssetTypeActions/AssetTypeActions_StaticMesh.cpp b/Engine/Source/Developer/AssetTools/Private/AssetTypeActions/AssetTypeActions_StaticMesh.cpp
--- a/Engine/Source/Developer/AssetTools/Private/AssetTypeActions/AssetTypeActions_StaticMesh.cpp
+++ b/Engine/Source/Developer/AssetTools/Private/AssetTypeActions/AssetTypeActions_StaticMesh.cpp
@@ -213,7 +213,7 @@ void FAssetTypeActions_StaticMesh::ExecutePasteLODSettings(TArray<TWeakObjectPtr
 
 			while (Mesh->SourceModels.Num() < LODCount)
 			{
-				new(Mesh->SourceModels) FStaticMeshSourceModel();
+				Mesh->AddSourceModel();
 			}
 
 			for (int32 i = 0; i < LODCount; i++)
diff --git a/Engine/Source/Developer/MeshDescriptionOperations/Private/MeshDescriptionOperations.cpp b/Engine/Source/Developer/MeshDescriptionOperations/Private/MeshDescriptionOperations.cpp
--- a/Engine/Source/Developer/MeshDescriptionOperations/Private/MeshDescriptionOperations.cpp
+++ b/Engine/Source/Developer/MeshDescriptionOperations/Private/MeshDescriptionOperations.cpp
@@ -237,7 +237,7 @@ void FMeshDescriptionOperations::ConvertSmoothGroupToHardEdges(const struct FRaw
 	}
 }
 
-void FMeshDescriptionOperations::ConverToRawMesh(const UMeshDescription* SourceMeshDescription, struct FRawMesh &DestinationRawMesh)
+void FMeshDescriptionOperations::ConverToRawMesh(const UMeshDescription* SourceMeshDescription, struct FRawMesh &DestinationRawMesh, const TMap<FName, int32>& MaterialMap)
 {
 	DestinationRawMesh.Empty();
 
@@ -250,6 +250,7 @@ void FMeshDescriptionOperations::ConverToRawMesh(const UMeshDescription* SourceM
 	const TVertexInstanceAttributeArray<FVector4>& VertexInstanceColors = SourceMeshDescription->VertexInstanceAttributes().GetAttributes<FVector4>(MeshAttribute::VertexInstance::Color);
 	const TVertexInstanceAttributeIndicesArray<FVector2D>& VertexInstanceUVs = SourceMeshDescription->VertexInstanceAttributes().GetAttributesSet<FVector2D>(MeshAttribute::VertexInstance::TextureCoordinate);
 
+	const TPolygonGroupAttributeArray<FName>& PolygonGroupMaterialSlotName = SourceMeshDescription->PolygonGroupAttributes().GetAttributes<FName>(MeshAttribute::PolygonGroup::ImportedMaterialSlotName);
 	const TPolygonGroupAttributeArray<int>& PolygonGroupMaterialIndex = SourceMeshDescription->PolygonGroupAttributes().GetAttributes<int>(MeshAttribute::PolygonGroup::MaterialIndex);
 
 	DestinationRawMesh.VertexPositions.AddZeroed(SourceMeshDescription->Vertices().Num());
@@ -278,6 +279,7 @@ void FMeshDescriptionOperations::ConverToRawMesh(const UMeshDescription* SourceM
 	DestinationRawMesh.FaceMaterialIndices.AddZeroed(TriangleNumber);
 	DestinationRawMesh.FaceSmoothingMasks.AddZeroed(TriangleNumber);
 
+	int32 TriangleIndex = 0;
 	for (const FPolygonID& PolygonID : SourceMeshDescription->Polygons().GetElementIDs())
 	{
 		const FPolygonGroupID& PolygonGroupID = SourceMeshDescription->GetPolygonPolygonGroup(PolygonID);
@@ -285,8 +287,17 @@ void FMeshDescriptionOperations::ConverToRawMesh(const UMeshDescription* SourceM
 		const TArray<FMeshTriangle>& Triangles = SourceMeshDescription->GetPolygonTriangles(PolygonID);
 		for (const FMeshTriangle& MeshTriangle : Triangles)
 		{
-			DestinationRawMesh.FaceMaterialIndices[PolygonIDValue] = PolygonGroupMaterialIndex[PolygonGroupID];
-			DestinationRawMesh.FaceSmoothingMasks[PolygonIDValue] = 0; //Conversion of soft/hard to smooth mask is done after the geometry is converted
+			if (MaterialMap.Num() > 0)
+			{
+				//The map should always be valid for all polygongroup
+				ensure(MaterialMap.Contains(PolygonGroupMaterialSlotName[PolygonGroupID]));
+				DestinationRawMesh.FaceMaterialIndices[TriangleIndex] = MaterialMap[PolygonGroupMaterialSlotName[PolygonGroupID]];
+			}
+			else
+			{
+				DestinationRawMesh.FaceMaterialIndices[TriangleIndex] = PolygonGroupMaterialIndex[PolygonGroupID];
+			}
+			DestinationRawMesh.FaceSmoothingMasks[TriangleIndex] = 0; //Conversion of soft/hard to smooth mask is done after the geometry is converted
 			for (int32 Corner = 0; Corner < 3; ++Corner)
 			{
 				const FVertexInstanceID VertexInstanceID = MeshTriangle.GetVertexInstanceID(Corner);
@@ -301,13 +312,14 @@ void FMeshDescriptionOperations::ConverToRawMesh(const UMeshDescription* SourceM
 					DestinationRawMesh.WedgeTexCoords[UVIndex][VertexInstanceIDValue] = VertexInstanceUVs.GetArrayForIndex(UVIndex)[VertexInstanceID];
 				}
 			}
+			++TriangleIndex;
 		}
 	}
 	//Convert the smoothgroup
 	ConvertHardEdgesToSmoothGroup(SourceMeshDescription, DestinationRawMesh);
 }
 
-void FMeshDescriptionOperations::ConverFromRawMesh(const struct FRawMesh &SourceRawMesh, UMeshDescription* DestinationMeshDescription)
+void FMeshDescriptionOperations::ConverFromRawMesh(const struct FRawMesh &SourceRawMesh, UMeshDescription* DestinationMeshDescription, const TMap<int32, FName>& MaterialMap)
 {
 	check(DestinationMeshDescription != nullptr);
 	DestinationMeshDescription->Empty();
@@ -326,7 +338,7 @@ void FMeshDescriptionOperations::ConverFromRawMesh(const struct FRawMesh &Source
 	TVertexInstanceAttributeIndicesArray<FVector2D>& VertexInstanceUVs = DestinationMeshDescription->VertexInstanceAttributes().GetAttributesSet<FVector2D>(MeshAttribute::VertexInstance::TextureCoordinate);
 
 	TPolygonGroupAttributeArray<FName>& PolygonGroupImportedMaterialSlotNames = DestinationMeshDescription->PolygonGroupAttributes().GetAttributes<FName>(MeshAttribute::PolygonGroup::ImportedMaterialSlotName);
-	TPolygonGroupAttributeArray<int>& PolygonGroupMaterialIndex = DestinationMeshDescription->PolygonGroupAttributes().GetAttributes<int>(MeshAttribute::PolygonGroup::MaterialIndex);
+	TPolygonGroupAttributeArray<int>& PolygonGroupImportedMaterialSlotIndex = DestinationMeshDescription->PolygonGroupAttributes().GetAttributes<int>(MeshAttribute::PolygonGroup::MaterialIndex);
 
 	int32 NumTexCoords = 0;
 	int32 MaxTexCoords = MAX_MESH_TEXTURE_COORDS;
@@ -362,19 +374,38 @@ void FMeshDescriptionOperations::ConverFromRawMesh(const struct FRawMesh &Source
 
 		//PolygonGroup
 		FPolygonGroupID PolygonGroupID = FPolygonGroupID::Invalid;
+		FName PolygonGroupImportedMaterialSlotName = NAME_None;
 		int32 MaterialIndex = SourceRawMesh.FaceMaterialIndices[TriangleIndex];
-		for (const FPolygonGroupID& SearchPolygonGroupID : DestinationMeshDescription->PolygonGroups().GetElementIDs())
+
+		if (MaterialMap.Num() > 0)
+		{
+			ensure(MaterialMap.Contains(SourceRawMesh.FaceMaterialIndices[TriangleIndex]));
+			PolygonGroupImportedMaterialSlotName = MaterialMap[SourceRawMesh.FaceMaterialIndices[TriangleIndex]];
+			for (const FPolygonGroupID& SearchPolygonGroupID : DestinationMeshDescription->PolygonGroups().GetElementIDs())
+			{
+				if (PolygonGroupImportedMaterialSlotNames[SearchPolygonGroupID] == PolygonGroupImportedMaterialSlotName)
+				{
+					PolygonGroupID = SearchPolygonGroupID;
+				}
+			}
+		}
+		else
 		{
-			if (PolygonGroupMaterialIndex[SearchPolygonGroupID] == MaterialIndex)
+			for (const FPolygonGroupID& SearchPolygonGroupID : DestinationMeshDescription->PolygonGroups().GetElementIDs())
 			{
-				PolygonGroupID = SearchPolygonGroupID;
+				if (PolygonGroupImportedMaterialSlotIndex[SearchPolygonGroupID] == MaterialIndex)
+				{
+					PolygonGroupID = SearchPolygonGroupID;
+				}
 			}
 		}
+		
+		
 		if (PolygonGroupID == FPolygonGroupID::Invalid)
 		{
 			PolygonGroupID = DestinationMeshDescription->CreatePolygonGroup();
-			PolygonGroupMaterialIndex[PolygonGroupID] = MaterialIndex;
-			PolygonGroupImportedMaterialSlotNames[PolygonGroupID] = FName(*FString::Printf(TEXT("MaterialSlot_%d"), MaterialIndex));
+			PolygonGroupImportedMaterialSlotNames[PolygonGroupID] = PolygonGroupImportedMaterialSlotName == NAME_None ? FName(*FString::Printf(TEXT("MaterialSlot_%d"), MaterialIndex)) : PolygonGroupImportedMaterialSlotName;
+			PolygonGroupImportedMaterialSlotIndex[PolygonGroupID] = SourceRawMesh.FaceMaterialIndices[TriangleIndex];
 			PolygonGroups.Add(PolygonGroupID);
 		}
 
diff --git a/Engine/Source/Developer/MeshDescriptionOperations/Public/MeshDescriptionOperations.h b/Engine/Source/Developer/MeshDescriptionOperations/Public/MeshDescriptionOperations.h
--- a/Engine/Source/Developer/MeshDescriptionOperations/Public/MeshDescriptionOperations.h
+++ b/Engine/Source/Developer/MeshDescriptionOperations/Public/MeshDescriptionOperations.h
@@ -34,10 +34,10 @@ public:
 	};
 
 	/** Convert this mesh description into the old FRawMesh format. */
-	static void ConverToRawMesh(const class UMeshDescription* SourceMeshDescription, struct FRawMesh &DestinationRawMesh);
+	static void ConverToRawMesh(const class UMeshDescription* SourceMeshDescription, struct FRawMesh &DestinationRawMesh, const TMap<FName, int32>& MaterialMap);
 
 	/** Convert old FRawMesh format to MeshDescription. */
-	static void ConverFromRawMesh(const struct FRawMesh &SourceRawMesh, class UMeshDescription* DestinationMeshDescription);
+	static void ConverFromRawMesh(const struct FRawMesh &SourceRawMesh, class UMeshDescription* DestinationMeshDescription, const TMap<int32, FName>& MaterialMap);
 
 	/**
 	 * Compute normal, tangent and Bi-Normal for every polygon in the mesh description. (this do not compute Vertex NTBs)
diff --git a/Engine/Source/Developer/MeshMergeUtilities/Private/MeshMergeUtilities.cpp b/Engine/Source/Developer/MeshMergeUtilities/Private/MeshMergeUtilities.cpp
--- a/Engine/Source/Developer/MeshMergeUtilities/Private/MeshMergeUtilities.cpp
+++ b/Engine/Source/Developer/MeshMergeUtilities/Private/MeshMergeUtilities.cpp
@@ -2251,19 +2251,19 @@ void FMeshMergeUtilities::MergeComponentsToStaticMesh(const TArray<UPrimitiveCom
 			FRawMesh& MergedMeshLOD = MergedRawMeshes[LODIndex];
 			if (MergedMeshLOD.VertexPositions.Num() > 0)
 			{
-				FStaticMeshSourceModel* SrcModel = new (StaticMesh->SourceModels) FStaticMeshSourceModel();
+				FStaticMeshSourceModel& SrcModel = StaticMesh->AddSourceModel();
 				// Don't allow the engine to recalculate normals
-				SrcModel->BuildSettings.bRecomputeNormals = false;
-				SrcModel->BuildSettings.bRecomputeTangents = false;
-				SrcModel->BuildSettings.bRemoveDegenerates = false;
-				SrcModel->BuildSettings.bUseHighPrecisionTangentBasis = false;
-				SrcModel->BuildSettings.bUseFullPrecisionUVs = false;
-				SrcModel->BuildSettings.bGenerateLightmapUVs = InSettings.bGenerateLightMapUV;
-				SrcModel->BuildSettings.MinLightmapResolution = InSettings.bComputedLightMapResolution ? DataTracker.GetLightMapDimension() : InSettings.TargetLightMapResolution;
-				SrcModel->BuildSettings.SrcLightmapIndex = 0;
-				SrcModel->BuildSettings.DstLightmapIndex = LightMapUVChannel;
-
-				SrcModel->SaveRawMesh(MergedMeshLOD);
+				SrcModel.BuildSettings.bRecomputeNormals = false;
+				SrcModel.BuildSettings.bRecomputeTangents = false;
+				SrcModel.BuildSettings.bRemoveDegenerates = false;
+				SrcModel.BuildSettings.bUseHighPrecisionTangentBasis = false;
+				SrcModel.BuildSettings.bUseFullPrecisionUVs = false;
+				SrcModel.BuildSettings.bGenerateLightmapUVs = InSettings.bGenerateLightMapUV;
+				SrcModel.BuildSettings.MinLightmapResolution = InSettings.bComputedLightMapResolution ? DataTracker.GetLightMapDimension() : InSettings.TargetLightMapResolution;
+				SrcModel.BuildSettings.SrcLightmapIndex = 0;
+				SrcModel.BuildSettings.DstLightmapIndex = LightMapUVChannel;
+
+				SrcModel.SaveRawMesh(MergedMeshLOD);
 			}
 		}
 		
diff --git a/Engine/Source/Developer/MeshMergeUtilities/Private/ProxyGenerationProcessor.cpp b/Engine/Source/Developer/MeshMergeUtilities/Private/ProxyGenerationProcessor.cpp
--- a/Engine/Source/Developer/MeshMergeUtilities/Private/ProxyGenerationProcessor.cpp
+++ b/Engine/Source/Developer/MeshMergeUtilities/Private/ProxyGenerationProcessor.cpp
@@ -173,14 +173,14 @@ void FProxyGenerationProcessor::ProcessJob(const FGuid& JobGuid, FProxyGeneratio
 	StaticMesh->LightMapResolution = Data->MergeData->InProxySettings.LightMapResolution;
 	StaticMesh->LightMapCoordinateIndex = 1;
 
-	FStaticMeshSourceModel* SrcModel = new (StaticMesh->SourceModels) FStaticMeshSourceModel();
+	FStaticMeshSourceModel& SrcModel = StaticMesh->AddSourceModel();
 	/*Don't allow the engine to recalculate normals*/
-	SrcModel->BuildSettings.bRecomputeNormals = false;
-	SrcModel->BuildSettings.bRecomputeTangents = false;
-	SrcModel->BuildSettings.bRemoveDegenerates = true;
-	SrcModel->BuildSettings.bUseHighPrecisionTangentBasis = false;
-	SrcModel->BuildSettings.bUseFullPrecisionUVs = false;
-	SrcModel->SaveRawMesh(Data->RawMesh);
+	SrcModel.BuildSettings.bRecomputeNormals = false;
+	SrcModel.BuildSettings.bRecomputeTangents = false;
+	SrcModel.BuildSettings.bRemoveDegenerates = true;
+	SrcModel.BuildSettings.bUseHighPrecisionTangentBasis = false;
+	SrcModel.BuildSettings.bUseFullPrecisionUVs = false;
+	SrcModel.SaveRawMesh(Data->RawMesh);
 
 	//Assign the proxy material to the static mesh
 	StaticMesh->StaticMaterials.Add(FStaticMaterial(ProxyMaterial));
diff --git a/Engine/Source/Developer/MeshUtilities/Private/MeshUtilities.cpp b/Engine/Source/Developer/MeshUtilities/Private/MeshUtilities.cpp
--- a/Engine/Source/Developer/MeshUtilities/Private/MeshUtilities.cpp
+++ b/Engine/Source/Developer/MeshUtilities/Private/MeshUtilities.cpp
@@ -888,16 +888,16 @@ UStaticMesh* FMeshUtilities::ConvertMeshesToStaticMesh(const TArray<UMeshCompone
 			{
 				if (RawMesh.IsValidOrFixable())
 				{
-					FStaticMeshSourceModel* SrcModel = new (StaticMesh->SourceModels) FStaticMeshSourceModel();
-					SrcModel->BuildSettings.bRecomputeNormals = false;
-					SrcModel->BuildSettings.bRecomputeTangents = false;
-					SrcModel->BuildSettings.bRemoveDegenerates = true;
-					SrcModel->BuildSettings.bUseHighPrecisionTangentBasis = false;
-					SrcModel->BuildSettings.bUseFullPrecisionUVs = false;
-					SrcModel->BuildSettings.bGenerateLightmapUVs = true;
-					SrcModel->BuildSettings.SrcLightmapIndex = 0;
-					SrcModel->BuildSettings.DstLightmapIndex = LightMapIndex;
-					SrcModel->SaveRawMesh(RawMesh);
+					FStaticMeshSourceModel& SrcModel = StaticMesh->AddSourceModel();
+					SrcModel.BuildSettings.bRecomputeNormals = false;
+					SrcModel.BuildSettings.bRecomputeTangents = false;
+					SrcModel.BuildSettings.bRemoveDegenerates = true;
+					SrcModel.BuildSettings.bUseHighPrecisionTangentBasis = false;
+					SrcModel.BuildSettings.bUseFullPrecisionUVs = false;
+					SrcModel.BuildSettings.bGenerateLightmapUVs = true;
+					SrcModel.BuildSettings.SrcLightmapIndex = 0;
+					SrcModel.BuildSettings.DstLightmapIndex = LightMapIndex;
+					SrcModel.SaveRawMesh(RawMesh);
 				}
 			}
 
diff --git a/Engine/Source/Editor/StaticMeshEditor/Private/StaticMeshEditorTools.cpp b/Engine/Source/Editor/StaticMeshEditor/Private/StaticMeshEditorTools.cpp
--- a/Engine/Source/Editor/StaticMeshEditor/Private/StaticMeshEditorTools.cpp
+++ b/Engine/Source/Editor/StaticMeshEditor/Private/StaticMeshEditorTools.cpp
@@ -3390,7 +3390,7 @@ void FLevelOfDetailSettingsLayout::ApplyChanges()
 	}
 	while (StaticMesh->SourceModels.Num() < LODCount)
 	{
-		new(StaticMesh->SourceModels) FStaticMeshSourceModel();
+		StaticMesh->AddSourceModel();
 	}
 	check(StaticMesh->SourceModels.Num() == LODCount);
 
diff --git a/Engine/Source/Editor/UnrealEd/Private/Fbx/FbxStaticMeshImport.cpp b/Engine/Source/Editor/UnrealEd/Private/Fbx/FbxStaticMeshImport.cpp
--- a/Engine/Source/Editor/UnrealEd/Private/Fbx/FbxStaticMeshImport.cpp
+++ b/Engine/Source/Editor/UnrealEd/Private/Fbx/FbxStaticMeshImport.cpp
@@ -1603,7 +1603,7 @@ UStaticMesh* UnFbx::FFbxImporter::ImportStaticMeshAsSingle(UObject* InParent, TA
 	if (StaticMesh->SourceModels.Num() < LODIndex+1)
 	{
 		// Add one LOD 
-		new(StaticMesh->SourceModels) FStaticMeshSourceModel();
+		StaticMesh->AddSourceModel();
 		
 		if (StaticMesh->SourceModels.Num() < LODIndex+1)
 		{
@@ -1661,15 +1661,10 @@ UStaticMesh* UnFbx::FFbxImporter::ImportStaticMeshAsSingle(UObject* InParent, TA
 		}
 	}
 
-	if (GetDefault<UEditorExperimentalSettings>()->bUseMeshDescription)
-	{
-		// SetOriginalMeshDescription will generate the rawmesh data
-		StaticMesh->SetOriginalMeshDescription(LODIndex, MeshDescription);
-	}
-	else
+	if (!GetDefault<UEditorExperimentalSettings>()->bUseMeshDescription)
 	{
 		// Store the new raw mesh.
-		SrcModel.SaveRawMesh(NewRawMesh);
+		SrcModel.SaveRawMesh(NewRawMesh, false);
 	}
 
 	if (bBuildStatus)
@@ -1838,7 +1833,7 @@ UStaticMesh* UnFbx::FFbxImporter::ImportStaticMeshAsSingle(UObject* InParent, TA
 					}
 				}
 
-				SrcModel.SaveRawMesh(LocalRawMesh);
+				SrcModel.SaveRawMesh(LocalRawMesh, false);
 			}
 
 			// Setup per-section info and the materials array.
@@ -1893,6 +1888,8 @@ UStaticMesh* UnFbx::FFbxImporter::ImportStaticMeshAsSingle(UObject* InParent, TA
 
 			FRawMesh LocalRawMesh;
 			SrcModel.RawMeshBulkData->LoadRawMesh(LocalRawMesh);
+			//This create the MeshDescription
+			SrcModel.SaveRawMesh(LocalRawMesh);
 		}
 		else
 		{
@@ -1950,7 +1947,7 @@ UStaticMesh* UnFbx::FFbxImporter::ImportStaticMeshAsSingle(UObject* InParent, TA
 			int32 NumLODs = LODGroup.GetDefaultNumLODs();
 			while (StaticMesh->SourceModels.Num() < NumLODs)
 			{
-				new (StaticMesh->SourceModels) FStaticMeshSourceModel();
+				StaticMesh->AddSourceModel();
 			}
 			for (int32 ModelLODIndex = 0; ModelLODIndex < NumLODs; ++ModelLODIndex)
 			{
@@ -2295,7 +2292,7 @@ void UnFbx::FFbxImporter::PostImportStaticMesh(UStaticMesh* StaticMesh, TArray<F
 			//Add missing LODs
 			while (StaticMesh->SourceModels.Num() < LODCount)
 			{
-				FStaticMeshSourceModel* SrcModel = new(StaticMesh->SourceModels) FStaticMeshSourceModel();
+				StaticMesh->AddSourceModel();
 			}
 		}
 	}
diff --git a/Engine/Source/Editor/UnrealEd/Private/StaticMeshEdit.cpp b/Engine/Source/Editor/UnrealEd/Private/StaticMeshEdit.cpp
--- a/Engine/Source/Editor/UnrealEd/Private/StaticMeshEdit.cpp
+++ b/Engine/Source/Editor/UnrealEd/Private/StaticMeshEdit.cpp
@@ -585,7 +585,7 @@ UStaticMesh* CreateStaticMesh(UMeshDescription* MeshDescription, TArray<FStaticM
 	auto StaticMesh = NewObject<UStaticMesh>(InOuter, InName, RF_Public | RF_Standalone);
 
 	// Add one LOD for the base mesh
-	FStaticMeshSourceModel* SrcModel = new(StaticMesh->SourceModels) FStaticMeshSourceModel();
+	FStaticMeshSourceModel& SrcModel = StaticMesh->AddSourceModel();
 	StaticMesh->SetOriginalMeshDescription(StaticMesh->SourceModels.Num() - 1, MeshDescription);
 	StaticMesh->StaticMaterials = Materials;
 
@@ -619,8 +619,8 @@ UStaticMesh* CreateStaticMesh(struct FRawMesh& RawMesh,TArray<FStaticMaterial>&
 	auto StaticMesh = NewObject<UStaticMesh>(InOuter, InName, RF_Public | RF_Standalone);
 
 	// Add one LOD for the base mesh
-	FStaticMeshSourceModel* SrcModel = new(StaticMesh->SourceModels) FStaticMeshSourceModel();
-	SrcModel->SaveRawMesh(RawMesh);
+	FStaticMeshSourceModel& SrcModel = StaticMesh->AddSourceModel();
+	SrcModel.SaveRawMesh(RawMesh);
 	StaticMesh->StaticMaterials = Materials;
 
 	int32 NumSections = StaticMesh->StaticMaterials.Num();
@@ -736,7 +736,7 @@ void GetBrushMesh(ABrush* Brush, UModel* Model, UMeshDescription* MeshDescriptio
 
 		// Find a material index for this polygon.
 		UMaterialInterface*	Material = Polygon.Material;
-		int32 MaterialIndex = OutMaterials.AddUnique(FStaticMaterial(Material));
+		int32 MaterialIndex = OutMaterials.AddUnique(FStaticMaterial(Material, Material->GetFName(), Material->GetFName()));
 		FPolygonGroupID CurrentPolygonGroupID = FPolygonGroupID::Invalid;
 		for (const FPolygonGroupID& PolygonGroupID : MeshDescription->PolygonGroups().GetElementIDs())
 		{
@@ -750,6 +750,7 @@ void GetBrushMesh(ABrush* Brush, UModel* Model, UMeshDescription* MeshDescriptio
 		{
 			CurrentPolygonGroupID = MeshDescription->CreatePolygonGroup();
 			PolygonGroupMaterialIndex[CurrentPolygonGroupID] = MaterialIndex;
+			PolygonGroupImportedMaterialSlotNames[CurrentPolygonGroupID] = Material->GetFName();
 		}
 
 		// Cache the texture coordinate system for this polygon.
@@ -1301,7 +1302,7 @@ void RestoreExistingMeshSettings(ExistingStaticMeshData* ExistingMesh, UStaticMe
 		{
 			if (NewMesh->SourceModels.Num() <= i)
 			{
-				new (NewMesh->SourceModels) FStaticMeshSourceModel();
+				NewMesh->AddSourceModel();
 			}
 
 			NewMesh->SourceModels[i].ReductionSettings = ExistingMesh->ExistingLODData[i].ExistingReductionSettings;
@@ -1633,18 +1634,18 @@ void RestoreExistingMeshData(ExistingStaticMeshData* ExistingMeshDataPtr, UStati
 
 	for(int32 i=NumCommonLODs; i < ExistingMeshDataPtr->ExistingLODData.Num(); ++i)
 	{
-		FStaticMeshSourceModel* SrcModel = new(NewMesh->SourceModels) FStaticMeshSourceModel();
+		FStaticMeshSourceModel& SrcModel = NewMesh->AddSourceModel();
 		if (ExistingMeshDataPtr->ExistingLODData[i].ExistingMeshDescription != nullptr)
 		{
-			SrcModel->OriginalMeshDescription = ExistingMeshDataPtr->ExistingLODData[i].ExistingMeshDescription;
+			SrcModel.OriginalMeshDescription = ExistingMeshDataPtr->ExistingLODData[i].ExistingMeshDescription;
 		}
 		if (ExistingMeshDataPtr->ExistingLODData[i].ExistingRawMesh.IsValidOrFixable())
 		{
-			SrcModel->RawMeshBulkData->SaveRawMesh(ExistingMeshDataPtr->ExistingLODData[i].ExistingRawMesh);
+			SrcModel.RawMeshBulkData->SaveRawMesh(ExistingMeshDataPtr->ExistingLODData[i].ExistingRawMesh);
 		}
-		SrcModel->BuildSettings = ExistingMeshDataPtr->ExistingLODData[i].ExistingBuildSettings;
-		SrcModel->ReductionSettings = ExistingMeshDataPtr->ExistingLODData[i].ExistingReductionSettings;
-		SrcModel->ScreenSize = ExistingMeshDataPtr->ExistingLODData[i].ExistingScreenSize;
+		SrcModel.BuildSettings = ExistingMeshDataPtr->ExistingLODData[i].ExistingBuildSettings;
+		SrcModel.ReductionSettings = ExistingMeshDataPtr->ExistingLODData[i].ExistingReductionSettings;
+		SrcModel.ScreenSize = ExistingMeshDataPtr->ExistingLODData[i].ExistingScreenSize;
 	}
 
 	// Restore the section info
diff --git a/Engine/Source/Editor/WorldBrowser/Private/Tiles/WorldTileCollectionModel.cpp b/Engine/Source/Editor/WorldBrowser/Private/Tiles/WorldTileCollectionModel.cpp
--- a/Engine/Source/Editor/WorldBrowser/Private/Tiles/WorldTileCollectionModel.cpp
+++ b/Engine/Source/Editor/WorldBrowser/Private/Tiles/WorldTileCollectionModel.cpp
@@ -2144,14 +2144,14 @@ bool FWorldTileCollectionModel::GenerateLODLevels(FLevelModelList InLevelList, i
 				StaticMesh->LightMapResolution = 64;
 				StaticMesh->LightMapCoordinateIndex = 1;
 
-				FStaticMeshSourceModel* SrcModel = new (StaticMesh->SourceModels) FStaticMeshSourceModel();
+				FStaticMeshSourceModel& SrcModel = StaticMesh->AddSourceModel();
 				/*Don't allow the engine to recalculate normals*/
-				SrcModel->BuildSettings.bRecomputeNormals = false;
-				SrcModel->BuildSettings.bRecomputeTangents = false;
-				SrcModel->BuildSettings.bRemoveDegenerates = false;
-				SrcModel->BuildSettings.bUseHighPrecisionTangentBasis = false;
-				SrcModel->BuildSettings.bUseFullPrecisionUVs = false;
-				SrcModel->SaveRawMesh(LandscapeRawMesh);
+				SrcModel.BuildSettings.bRecomputeNormals = false;
+				SrcModel.BuildSettings.bRecomputeTangents = false;
+				SrcModel.BuildSettings.bRemoveDegenerates = false;
+				SrcModel.BuildSettings.bUseHighPrecisionTangentBasis = false;
+				SrcModel.BuildSettings.bUseFullPrecisionUVs = false;
+				SrcModel.SaveRawMesh(LandscapeRawMesh);
 
 				//Assign the proxy material to the static mesh
 				StaticMesh->StaticMaterials.Add(FStaticMaterial(StaticLandscapeMaterial));
diff --git a/Engine/Source/Runtime/Engine/Classes/Engine/StaticMesh.h b/Engine/Source/Runtime/Engine/Classes/Engine/StaticMesh.h
--- a/Engine/Source/Runtime/Engine/Classes/Engine/StaticMesh.h
+++ b/Engine/Source/Runtime/Engine/Classes/Engine/StaticMesh.h
@@ -153,7 +153,14 @@ struct FStaticMeshSourceModel
 #if WITH_EDITOR
 	/** Imported raw mesh data. Optional for all but the first LOD. */
 	class FRawMeshBulkData* RawMeshBulkData;
-
+	
+	/*
+	 * The staticmesh owner of this source model. We need the SM to be able to convert between MeshDesription and RawMesh.
+	 * RawMesh use int32 material index and MeshDescription use FName material slot name.
+	 * This memeber is fill in the PostLoad of the static mesh.
+	 * TODO: Remove this member when FRawMesh will be remove.
+	 */
+	class UStaticMesh* StaticMeshOwner;
 	/*
 	 * Accessor to Load and save the raw mesh or the mesh description depending on the editor settings.
 	 * Temporary until we deprecate the RawMesh. It use UEditorExperimentalSettings::bUseMeshDescription to know which structure we use to save/load the data.
@@ -161,7 +168,7 @@ struct FStaticMeshSourceModel
 	 */
 	ENGINE_API bool IsRawMeshEmpty() const;
 	ENGINE_API void LoadRawMesh(struct FRawMesh& OutRawMesh) const;
-	ENGINE_API void SaveRawMesh(struct FRawMesh& InRawMesh);
+	ENGINE_API void SaveRawMesh(struct FRawMesh& InRawMesh, bool bConvertToMeshdescription = true);
 
 #endif // #if WITH_EDITOR
 
@@ -703,6 +710,9 @@ public:
 
 	FOnExtendedBoundsChanged& GetOnExtendedBoundsChanged() { return OnExtendedBoundsChanged; }
 	FOnMeshChanged& GetOnMeshChanged() { return OnMeshChanged; }
+
+	//SourceModels API
+	ENGINE_API FStaticMeshSourceModel& AddSourceModel();
 #endif // WITH_EDITOR
 
 	ENGINE_API virtual void Serialize(FArchive& Ar) override;
diff --git a/Engine/Source/Runtime/Engine/Private/StaticMesh.cpp b/Engine/Source/Runtime/Engine/Private/StaticMesh.cpp
--- a/Engine/Source/Runtime/Engine/Private/StaticMesh.cpp
+++ b/Engine/Source/Runtime/Engine/Private/StaticMesh.cpp
@@ -2019,6 +2019,14 @@ void UStaticMesh::BroadcastNavCollisionChange()
 		}
 	}
 }
+
+FStaticMeshSourceModel& UStaticMesh::AddSourceModel()
+{
+	int32 LodModelIndex = SourceModels.AddDefaulted();
+	SourceModels[LodModelIndex].StaticMeshOwner = this;
+	return SourceModels[LodModelIndex];
+}
+
 #endif // WITH_EDITOR
 
 void UStaticMesh::BeginDestroy()
@@ -2150,6 +2158,7 @@ FStaticMeshSourceModel::FStaticMeshSourceModel()
 	RawMeshBulkData = new FRawMeshBulkData();
 	ScreenSize = 0.0f;
 	OriginalMeshDescription = nullptr;
+	StaticMeshOwner = nullptr;
 #endif // #if WITH_EDITOR
 }
 
@@ -2162,6 +2171,7 @@ FStaticMeshSourceModel::~FStaticMeshSourceModel()
 		RawMeshBulkData = NULL;
 	}
 	OriginalMeshDescription = nullptr;
+	StaticMeshOwner = nullptr;
 #endif // #if WITH_EDITOR
 }
 
@@ -2175,7 +2185,13 @@ void FStaticMeshSourceModel::LoadRawMesh(FRawMesh& OutRawMesh) const
 {
 	if (RawMeshBulkData->IsEmpty() && OriginalMeshDescription != nullptr)
 	{
-		FMeshDescriptionOperations::ConverToRawMesh(OriginalMeshDescription, OutRawMesh);
+		TMap<FName, int32> MaterialMap;
+		check(StaticMeshOwner != nullptr);
+		for (int32 MaterialIndex = 0; MaterialIndex < StaticMeshOwner->StaticMaterials.Num(); ++MaterialIndex)
+		{
+			MaterialMap.Add(StaticMeshOwner->StaticMaterials[MaterialIndex].ImportedMaterialSlotName, MaterialIndex);
+		}
+		FMeshDescriptionOperations::ConverToRawMesh(OriginalMeshDescription, OutRawMesh, MaterialMap);
 	}
 	else
 	{
@@ -2183,7 +2199,7 @@ void FStaticMeshSourceModel::LoadRawMesh(FRawMesh& OutRawMesh) const
 	}
 }
 
-void FStaticMeshSourceModel::SaveRawMesh(FRawMesh& InRawMesh)
+void FStaticMeshSourceModel::SaveRawMesh(FRawMesh& InRawMesh, bool bConvertToMeshdescription /*= true*/)
 {
 	if (!InRawMesh.IsValid())
 	{
@@ -2191,9 +2207,15 @@ void FStaticMeshSourceModel::SaveRawMesh(FRawMesh& InRawMesh)
 	}
 	//Save both format
 	RawMeshBulkData->SaveRawMesh(InRawMesh);
-	if (OriginalMeshDescription != nullptr)
+	if (bConvertToMeshdescription && OriginalMeshDescription != nullptr)
 	{
-		FMeshDescriptionOperations::ConverFromRawMesh(InRawMesh, OriginalMeshDescription);
+		TMap<int32, FName> MaterialMap;
+		check(StaticMeshOwner != nullptr);
+		for (int32 MaterialIndex = 0; MaterialIndex < StaticMeshOwner->StaticMaterials.Num(); ++MaterialIndex)
+		{
+			MaterialMap.Add(MaterialIndex, StaticMeshOwner->StaticMaterials[MaterialIndex].ImportedMaterialSlotName);
+		}
+		FMeshDescriptionOperations::ConverFromRawMesh(InRawMesh, OriginalMeshDescription, MaterialMap);
 	}
 }
 
@@ -2391,7 +2413,12 @@ UMeshDescription* UStaticMesh::GetOriginalMeshDescription(int32 LodIndex)
 			SourceModels[LodIndex].OriginalMeshDescription = NewObject<UMeshDescription>(this, NAME_None);
 			FRawMesh LodRawMesh;
 			SourceModels[LodIndex].LoadRawMesh(LodRawMesh);
-			FMeshDescriptionOperations::ConverFromRawMesh(LodRawMesh, SourceModels[LodIndex].OriginalMeshDescription);
+			TMap<int32, FName> MaterialMap;
+			for (int32 MaterialIndex = 0; MaterialIndex < StaticMaterials.Num(); ++MaterialIndex)
+			{
+				MaterialMap.Add(MaterialIndex, StaticMaterials[MaterialIndex].ImportedMaterialSlotName);
+			}
+			FMeshDescriptionOperations::ConverFromRawMesh(LodRawMesh, SourceModels[LodIndex].OriginalMeshDescription, MaterialMap);
 		}
 		return SourceModels[LodIndex].OriginalMeshDescription;
 	}
@@ -2406,8 +2433,13 @@ void UStaticMesh::SetOriginalMeshDescription(int32 LodIndex, class UMeshDescript
 	if (MeshDescription != nullptr)
 	{
 		FRawMesh TempRawMesh;
+		TMap<FName, int32> MaterialMap;
+		for (int32 MaterialIndex = 0; MaterialIndex < StaticMaterials.Num(); ++MaterialIndex)
+		{
+			MaterialMap.Add(StaticMaterials[MaterialIndex].ImportedMaterialSlotName, MaterialIndex);
+		}
 		//Convert RawMesh to meshdescription
-		FMeshDescriptionOperations::ConverToRawMesh(MeshDescription, TempRawMesh);
+		FMeshDescriptionOperations::ConverToRawMesh(MeshDescription, TempRawMesh, MaterialMap);
 		SourceModels[LodIndex].RawMeshBulkData->SaveRawMesh(TempRawMesh);
 	}
 	else
@@ -2659,7 +2691,12 @@ void UStaticMesh::CacheMeshData()
 				//The original mesh description is null we must create one
 				SourceModel.OriginalMeshDescription = NewObject<UMeshDescription>(this, NAME_None);
 				//Convert the raw mesh to mesh description and save the mesh description in the ddc
-				FMeshDescriptionOperations::ConverFromRawMesh(TempRawMesh, SourceModel.OriginalMeshDescription);
+				TMap<int32, FName> MaterialMap;
+				for (int32 MaterialIndex = 0; MaterialIndex < StaticMaterials.Num(); ++MaterialIndex)
+				{
+					MaterialMap.Add(MaterialIndex, StaticMaterials[MaterialIndex].ImportedMaterialSlotName);
+				}
+				FMeshDescriptionOperations::ConverFromRawMesh(TempRawMesh, SourceModel.OriginalMeshDescription, MaterialMap);
 			}
 		}
 	}
@@ -3028,6 +3065,15 @@ void UStaticMesh::PostLoad()
 		{
 			DistanceFieldReplacementMesh->ConditionalPostLoad();
 		}
+		
+		//TODO remove this code when FRawMesh will be removed
+		//Fill the static mesh owner
+		int32 NumLODs = SourceModels.Num();
+		for (int32 LODIndex = 0; LODIndex < NumLODs; ++LODIndex)
+		{
+			FStaticMeshSourceModel& SrcModel = SourceModels[LODIndex];
+			SrcModel.StaticMeshOwner = this;
+		}
 	}
 
 	if (!GetOutermost()->HasAnyPackageFlags(PKG_FilterEditorOnly))
