Fix automation tests crash

diff --git a/Engine/Source/Developer/MeshDescriptionOperations/Private/MeshDescriptionOperations.cpp b/Engine/Source/Developer/MeshDescriptionOperations/Private/MeshDescriptionOperations.cpp
--- a/Engine/Source/Developer/MeshDescriptionOperations/Private/MeshDescriptionOperations.cpp
+++ b/Engine/Source/Developer/MeshDescriptionOperations/Private/MeshDescriptionOperations.cpp
@@ -251,7 +251,6 @@ void FMeshDescriptionOperations::ConverToRawMesh(const UMeshDescription* SourceM
 	const TVertexInstanceAttributeIndicesArray<FVector2D>& VertexInstanceUVs = SourceMeshDescription->VertexInstanceAttributes().GetAttributesSet<FVector2D>(MeshAttribute::VertexInstance::TextureCoordinate);
 
 	const TPolygonGroupAttributeArray<FName>& PolygonGroupMaterialSlotName = SourceMeshDescription->PolygonGroupAttributes().GetAttributes<FName>(MeshAttribute::PolygonGroup::ImportedMaterialSlotName);
-	const TPolygonGroupAttributeArray<int>& PolygonGroupMaterialIndex = SourceMeshDescription->PolygonGroupAttributes().GetAttributes<int>(MeshAttribute::PolygonGroup::MaterialIndex);
 
 	DestinationRawMesh.VertexPositions.AddZeroed(SourceMeshDescription->Vertices().Num());
 	for (const FVertexID& VertexID : SourceMeshDescription->Vertices().GetElementIDs())
@@ -287,15 +286,13 @@ void FMeshDescriptionOperations::ConverToRawMesh(const UMeshDescription* SourceM
 		const TArray<FMeshTriangle>& Triangles = SourceMeshDescription->GetPolygonTriangles(PolygonID);
 		for (const FMeshTriangle& MeshTriangle : Triangles)
 		{
-			if (MaterialMap.Num() > 0)
+			if (MaterialMap.Num() > 0 && MaterialMap.Contains(PolygonGroupMaterialSlotName[PolygonGroupID]))
 			{
-				//The map should always be valid for all polygongroup
-				ensure(MaterialMap.Contains(PolygonGroupMaterialSlotName[PolygonGroupID]));
 				DestinationRawMesh.FaceMaterialIndices[TriangleIndex] = MaterialMap[PolygonGroupMaterialSlotName[PolygonGroupID]];
 			}
 			else
 			{
-				DestinationRawMesh.FaceMaterialIndices[TriangleIndex] = PolygonGroupMaterialIndex[PolygonGroupID];
+				DestinationRawMesh.FaceMaterialIndices[TriangleIndex] = 0;
 			}
 			DestinationRawMesh.FaceSmoothingMasks[TriangleIndex] = 0; //Conversion of soft/hard to smooth mask is done after the geometry is converted
 			for (int32 Corner = 0; Corner < 3; ++Corner)
@@ -338,7 +335,6 @@ void FMeshDescriptionOperations::ConverFromRawMesh(const struct FRawMesh &Source
 	TVertexInstanceAttributeIndicesArray<FVector2D>& VertexInstanceUVs = DestinationMeshDescription->VertexInstanceAttributes().GetAttributesSet<FVector2D>(MeshAttribute::VertexInstance::TextureCoordinate);
 
 	TPolygonGroupAttributeArray<FName>& PolygonGroupImportedMaterialSlotNames = DestinationMeshDescription->PolygonGroupAttributes().GetAttributes<FName>(MeshAttribute::PolygonGroup::ImportedMaterialSlotName);
-	TPolygonGroupAttributeArray<int>& PolygonGroupImportedMaterialSlotIndex = DestinationMeshDescription->PolygonGroupAttributes().GetAttributes<int>(MeshAttribute::PolygonGroup::MaterialIndex);
 
 	int32 NumTexCoords = 0;
 	int32 MaxTexCoords = MAX_MESH_TEXTURE_COORDS;
@@ -377,9 +373,8 @@ void FMeshDescriptionOperations::ConverFromRawMesh(const struct FRawMesh &Source
 		FName PolygonGroupImportedMaterialSlotName = NAME_None;
 		int32 MaterialIndex = SourceRawMesh.FaceMaterialIndices[TriangleIndex];
 
-		if (MaterialMap.Num() > 0)
+		if (MaterialMap.Num() > 0 && MaterialMap.Contains(SourceRawMesh.FaceMaterialIndices[TriangleIndex]))
 		{
-			ensure(MaterialMap.Contains(SourceRawMesh.FaceMaterialIndices[TriangleIndex]));
 			PolygonGroupImportedMaterialSlotName = MaterialMap[SourceRawMesh.FaceMaterialIndices[TriangleIndex]];
 			for (const FPolygonGroupID& SearchPolygonGroupID : DestinationMeshDescription->PolygonGroups().GetElementIDs())
 			{
@@ -389,23 +384,11 @@ void FMeshDescriptionOperations::ConverFromRawMesh(const struct FRawMesh &Source
 				}
 			}
 		}
-		else
-		{
-			for (const FPolygonGroupID& SearchPolygonGroupID : DestinationMeshDescription->PolygonGroups().GetElementIDs())
-			{
-				if (PolygonGroupImportedMaterialSlotIndex[SearchPolygonGroupID] == MaterialIndex)
-				{
-					PolygonGroupID = SearchPolygonGroupID;
-				}
-			}
-		}
-		
 		
 		if (PolygonGroupID == FPolygonGroupID::Invalid)
 		{
 			PolygonGroupID = DestinationMeshDescription->CreatePolygonGroup();
 			PolygonGroupImportedMaterialSlotNames[PolygonGroupID] = PolygonGroupImportedMaterialSlotName == NAME_None ? FName(*FString::Printf(TEXT("MaterialSlot_%d"), MaterialIndex)) : PolygonGroupImportedMaterialSlotName;
-			PolygonGroupImportedMaterialSlotIndex[PolygonGroupID] = SourceRawMesh.FaceMaterialIndices[TriangleIndex];
 			PolygonGroups.Add(PolygonGroupID);
 		}
 
