Fix old faceted meshes, bump version to ensure static mesh are rebuild build
- Add some normals and tangents tests but they are disable until the new compute tangents return the same value then the original FRawMesh code.

diff --git a/Engine/Plugins/Tests/EditorTests/Source/EditorTests/Private/UnrealEd/MeshDescriptionTests.cpp b/Engine/Plugins/Tests/EditorTests/Source/EditorTests/Private/UnrealEd/MeshDescriptionTests.cpp
--- a/Engine/Plugins/Tests/EditorTests/Source/EditorTests/Private/UnrealEd/MeshDescriptionTests.cpp
+++ b/Engine/Plugins/Tests/EditorTests/Source/EditorTests/Private/UnrealEd/MeshDescriptionTests.cpp
@@ -13,6 +13,7 @@
 #include "MeshBuilder.h"
 #include "MeshAttributes.h"
 #include "RawMesh.h"
+#include "MeshUtilities.h"
 
 //////////////////////////////////////////////////////////////////////////
 
@@ -25,6 +26,7 @@
 IMPLEMENT_COMPLEX_AUTOMATION_TEST(FMeshDescriptionAutomationTest, "Editor.Meshes.MeshDescription", (EAutomationTestFlags::EditorContext | EAutomationTestFlags::NonNullRHI | EAutomationTestFlags::EngineFilter))
 
 #define ConversionTestData 1
+#define NTB_TestData 2
 
 class FMeshDescriptionTest
 {
@@ -38,6 +40,7 @@ public:
 	FString TestData;
 	bool Execute(FAutomationTestExecutionInfo& ExecutionInfo);
 	bool ConversionTest(FAutomationTestExecutionInfo& ExecutionInfo);
+	bool NTBTest(FAutomationTestExecutionInfo& ExecutionInfo);
 private:
 	bool CompareRawMesh(const FString& AssetName, FAutomationTestExecutionInfo& ExecutionInfo, const FRawMesh& ReferenceRawMesh, const FRawMesh& ResultRawMesh) const;
 	bool CompareMeshDescription(const FString& AssetName, FAutomationTestExecutionInfo& ExecutionInfo, const UMeshDescription* ReferenceMeshDescription, const UMeshDescription* MeshDescription) const;
@@ -87,11 +90,17 @@ private:
 void FMeshDescriptionAutomationTest::GetTests(TArray<FString>& OutBeautifiedNames, TArray <FString>& OutTestCommands) const
 {
 	FMeshDescriptionTests::GetInstance().ClearTests();
-	//Create cube test
+	//Create conversion test
 	FMeshDescriptionTest ConversionTest(TEXT("Conversion test data"), FString::FromInt(ConversionTestData));
 	FMeshDescriptionTests::GetInstance().AddTest(ConversionTest);
 	OutBeautifiedNames.Add(ConversionTest.BeautifiedNames);
 	OutTestCommands.Add(ConversionTest.TestData);
+	//Create NormalTangentBinormal test
+/*	FMeshDescriptionTest NTB_Test(TEXT("Normals Tangents and Binormals test"), FString::FromInt(NTB_TestData));
+	FMeshDescriptionTests::GetInstance().AddTest(NTB_Test);
+	OutBeautifiedNames.Add(NTB_Test.BeautifiedNames);
+	OutTestCommands.Add(NTB_Test.TestData);
+*/
 }
 
 /**
@@ -125,6 +134,9 @@ bool FMeshDescriptionTest::Execute(FAutomationTestExecutionInfo& ExecutionInfo)
 	case ConversionTestData:
 		bSuccess = ConversionTest(ExecutionInfo);
 		break;
+	case NTB_TestData:
+		bSuccess = NTBTest(ExecutionInfo);
+		break;
 	}
 	return bSuccess;
 }
@@ -484,9 +496,7 @@ bool FMeshDescriptionTest::ConversionTest(FAutomationTestExecutionInfo& Executio
 
 		if (AssetMesh != nullptr)
 		{
-#if WITH_EDITOR
 			AssetMesh->BuildCacheAutomationTestGuid = FGuid::NewGuid();
-#endif
 			TMap<FName, int32> MaterialMap;
 			TMap<int32, FName> MaterialMapInverse;
 			for (int32 MaterialIndex = 0; MaterialIndex < AssetMesh->StaticMaterials.Num(); ++MaterialIndex)
@@ -552,3 +562,160 @@ bool FMeshDescriptionTest::ConversionTest(FAutomationTestExecutionInfo& Executio
 	TryCollectGarbage(GARBAGE_COLLECTION_KEEPFLAGS);
 	return bAllSame;
 }
+
+bool FMeshDescriptionTest::NTBTest(FAutomationTestExecutionInfo& ExecutionInfo)
+{
+	TArray<FString> AssetNames;
+	AssetNames.Add(TEXT("Cone_1"));
+	AssetNames.Add(TEXT("Cone_2"));
+	AssetNames.Add(TEXT("Cube"));
+	AssetNames.Add(TEXT("Patch_1"));
+	AssetNames.Add(TEXT("Patch_2"));
+	AssetNames.Add(TEXT("Patch_3"));
+	AssetNames.Add(TEXT("Patch_4"));
+	AssetNames.Add(TEXT("Patch_5"));
+	AssetNames.Add(TEXT("Pentagone"));
+	AssetNames.Add(TEXT("Sphere_1"));
+	AssetNames.Add(TEXT("Sphere_2"));
+	AssetNames.Add(TEXT("Sphere_3"));
+	AssetNames.Add(TEXT("Torus_1"));
+	AssetNames.Add(TEXT("Torus_2"));
+
+	bool bAllSame = true;
+	for (const FString& AssetName : AssetNames)
+	{
+		FString FullAssetName = TEXT("/Game/Tests/MeshDescription/") + AssetName + TEXT(".") + AssetName;
+		UStaticMesh* AssetMesh = LoadObject<UStaticMesh>(nullptr, *FullAssetName, nullptr, LOAD_None, nullptr);
+
+		if (AssetMesh == nullptr)
+		{
+			continue;
+		}
+		//Dirty the build
+		AssetMesh->BuildCacheAutomationTestGuid = FGuid::NewGuid();
+		UMeshDescription* MeshDescription = AssetMesh->GetOriginalMeshDescription(0);
+		FRawMesh RawMesh;
+		AssetMesh->SourceModels[0].LoadRawMesh(RawMesh);
+
+		//const TVertexAttributeArray<FVector>& VertexPositions = MeshDescription->VertexAttributes().GetAttributes<FVector>(MeshAttribute::Vertex::Position);
+		const TVertexInstanceAttributeArray<FVector>& VertexInstanceNormals = MeshDescription->VertexInstanceAttributes().GetAttributes<FVector>(MeshAttribute::VertexInstance::Normal);
+		const TVertexInstanceAttributeArray<FVector>& VertexInstanceTangents = MeshDescription->VertexInstanceAttributes().GetAttributes<FVector>(MeshAttribute::VertexInstance::Tangent);
+		const TVertexInstanceAttributeArray<float>& VertexInstanceBinormalSigns = MeshDescription->VertexInstanceAttributes().GetAttributes<float>(MeshAttribute::VertexInstance::BinormalSign);
+		const TVertexInstanceAttributeArray<FVector4>& VertexInstanceColors = MeshDescription->VertexInstanceAttributes().GetAttributes<FVector4>(MeshAttribute::VertexInstance::Color);
+		const TVertexInstanceAttributeIndicesArray<FVector2D>& VertexInstanceUVs = MeshDescription->VertexInstanceAttributes().GetAttributesSet<FVector2D>(MeshAttribute::VertexInstance::TextureCoordinate);
+		int32 ExistingUVCount = VertexInstanceUVs.GetNumIndices();
+		//Build the normals and tangent and compare the result
+		MeshDescription->ComputePolygonTangentsAndNormals(SMALL_NUMBER);
+		EComputeNTBsOptions ComputeNTBsOptions = EComputeNTBsOptions::Normals | EComputeNTBsOptions::Tangents;
+		MeshDescription->ComputeTangentsAndNormals(ComputeNTBsOptions);
+		//FMeshDescriptionOperations::CreatePolygonNTB(MeshDescription, SMALL_NUMBER);
+		//FMeshDescriptionOperations::CreateNormals(MeshDescription, FMeshDescriptionOperations::ETangentOptions::BlendOverlappingNormals, true);
+
+		IMeshUtilities& MeshUtilities = FModuleManager::Get().LoadModuleChecked<IMeshUtilities>(TEXT("MeshUtilities"));
+		FMeshBuildSettings MeshBuildSettings;
+		MeshBuildSettings.bRemoveDegenerates = true;
+		MeshBuildSettings.bUseMikkTSpace = false;
+		MeshUtilities.RecomputeTangentsAndNormalsForRawMesh(true, true, MeshBuildSettings, RawMesh);
+
+		//The normals and tangents of both the meshdescription and the RawMesh should be equal to not break old data
+		if(RawMesh.WedgeIndices.Num() != MeshDescription->VertexInstances().Num())
+		{
+			ExecutionInfo.AddEvent(FAutomationEvent(EAutomationEventType::Error, FString::Printf(TEXT("Test: [Normals Tangents and Binormals test]    Asset: [%s]    Error: The number of vertex instances is not equal between FRawMesh [%d] and UMeshDescription [%d]."),
+				*AssetName,
+				RawMesh.WedgeIndices.Num(),
+				MeshDescription->VertexInstances().Num())));
+			continue;
+		}
+		int32 VertexInstanceIndex = 0;
+		int32 TriangleIndex = 0;
+
+		auto OutputError=[&ExecutionInfo, &AssetName](FString ErrorMessage)
+		{
+			ExecutionInfo.AddEvent(FAutomationEvent(EAutomationEventType::Error, FString::Printf(TEXT("Test: [Normals Tangents and Binormals test]    Asset: [%s]    Error: %s."),
+				*AssetName,
+				*ErrorMessage)));
+		};
+		bool bError = false;
+		for (const FPolygonID& PolygonID : MeshDescription->Polygons().GetElementIDs())
+		{
+			if (bError)
+			{
+				break;
+			}
+			const FPolygonGroupID& PolygonGroupID = MeshDescription->GetPolygonPolygonGroup(PolygonID);
+			int32 PolygonIDValue = PolygonID.GetValue();
+			const TArray<FMeshTriangle>& Triangles = MeshDescription->GetPolygonTriangles(PolygonID);
+			for (const FMeshTriangle& MeshTriangle : Triangles)
+			{
+				if (bError)
+				{
+					break;
+				}
+				for (int32 Corner = 0; Corner < 3 && bError == false; ++Corner)
+				{
+					uint32 WedgeIndex = TriangleIndex * 3 + Corner;
+					const FVertexInstanceID VertexInstanceID = MeshTriangle.GetVertexInstanceID(Corner);
+					const int32 VertexInstanceIDValue = VertexInstanceID.GetValue();
+					if (RawMesh.WedgeColors[WedgeIndex] != FLinearColor(VertexInstanceColors[VertexInstanceID]).ToFColor(true))
+					{
+						FString MeshDescriptionColor = FLinearColor(VertexInstanceColors[VertexInstanceID]).ToFColor(true).ToString();
+						OutputError(FString::Printf(TEXT("Vertex color is different between MeshDescription [%s] and FRawMesh [%s].   Indice[%d]")
+							, *MeshDescriptionColor
+							, *(RawMesh.WedgeColors[WedgeIndex].ToString())
+							, VertexInstanceIDValue));
+						bError = true;
+					}
+					if (RawMesh.WedgeIndices[WedgeIndex] != MeshDescription->GetVertexInstanceVertex(VertexInstanceID).GetValue())
+					{
+						OutputError(FString::Printf(TEXT("Vertex index is different between MeshDescription [%d] and FRawMesh [%d].   Indice[%d]")
+							, MeshDescription->GetVertexInstanceVertex(VertexInstanceID).GetValue()
+							, RawMesh.WedgeIndices[WedgeIndex]
+							, VertexInstanceIDValue));
+						bError = true;
+					}
+					if (!RawMesh.WedgeTangentX[WedgeIndex].Equals(VertexInstanceTangents[VertexInstanceID], THRESH_NORMALS_ARE_SAME))
+					{
+						OutputError(FString::Printf(TEXT("Vertex tangent is different between MeshDescription [%s] and FRawMesh [%s].   Indice[%d]")
+							, *(VertexInstanceTangents[VertexInstanceID].ToString())
+							, *(RawMesh.WedgeTangentX[WedgeIndex].ToString())
+							, VertexInstanceIDValue));
+						bError = true;
+					}
+					if (!RawMesh.WedgeTangentY[WedgeIndex].Equals(FVector::CrossProduct(VertexInstanceNormals[VertexInstanceID], VertexInstanceTangents[VertexInstanceID]).GetSafeNormal() * VertexInstanceBinormalSigns[VertexInstanceID], THRESH_NORMALS_ARE_SAME))
+					{
+						FVector MeshDescriptionBinormalResult = FVector::CrossProduct(VertexInstanceNormals[VertexInstanceID], VertexInstanceTangents[VertexInstanceID]).GetSafeNormal() * VertexInstanceBinormalSigns[VertexInstanceID];
+						OutputError(FString::Printf(TEXT("Vertex binormal is different between MeshDescription [%s] and FRawMesh [%s].   Indice[%d]")
+							, *(MeshDescriptionBinormalResult.ToString())
+							, *(RawMesh.WedgeTangentY[WedgeIndex].ToString())
+							, VertexInstanceIDValue));
+						bError = true;
+					}
+					if (!RawMesh.WedgeTangentZ[WedgeIndex].Equals(VertexInstanceNormals[VertexInstanceID], THRESH_NORMALS_ARE_SAME))
+					{
+						OutputError(FString::Printf(TEXT("Vertex normal is different between MeshDescription [%s] and FRawMesh [%s].   Indice[%d]")
+							, *(VertexInstanceNormals[VertexInstanceID].ToString())
+							, *(RawMesh.WedgeTangentZ[WedgeIndex].ToString())
+							, VertexInstanceIDValue));
+						bError = true;
+					}
+
+					for (int32 UVIndex = 0; UVIndex < ExistingUVCount; ++UVIndex)
+					{
+						if (!RawMesh.WedgeTexCoords[UVIndex][WedgeIndex].Equals(VertexInstanceUVs.GetArrayForIndex(UVIndex)[VertexInstanceID], THRESH_UVS_ARE_SAME))
+						{
+							OutputError(FString::Printf(TEXT("Vertex Texture coordinnate is different between MeshDescription [%s] and FRawMesh [%s].   UVIndex[%d]  Indice[%d]")
+								, *(VertexInstanceUVs.GetArrayForIndex(UVIndex)[VertexInstanceID].ToString())
+								, *(RawMesh.WedgeTexCoords[UVIndex][WedgeIndex].ToString())
+								, UVIndex
+								, VertexInstanceIDValue));
+							bError = true;
+						}
+					}
+				}
+				++TriangleIndex;
+			}
+		}
+	}
+
+	return bAllSame;
+}
\ No newline at end of file
diff --git a/Engine/Source/Developer/MeshBuilder/Private/MeshDescriptionHelper.cpp b/Engine/Source/Developer/MeshBuilder/Private/MeshDescriptionHelper.cpp
--- a/Engine/Source/Developer/MeshBuilder/Private/MeshDescriptionHelper.cpp
+++ b/Engine/Source/Developer/MeshBuilder/Private/MeshDescriptionHelper.cpp
@@ -67,7 +67,8 @@ UMeshDescription* FMeshDescriptionHelper::GetRenderMeshDescription(UObject* Owne
 		}
 
 		//This function make sure the Polygon NTB are compute and also remove degenerated triangle from the render mesh description.
-		RenderMeshDescription->ComputePolygonTangentsAndNormals(BuildSettings->bRemoveDegenerates ? SMALL_NUMBER : 0.0f);
+		FMeshDescriptionOperations::CreatePolygonNTB(RenderMeshDescription, ComparisonThreshold);
+		//RenderMeshDescription->ComputePolygonTangentsAndNormals(BuildSettings->bRemoveDegenerates ? SMALL_NUMBER : 0.0f);
 		
 		//Keep the original mesh description NTBs if we do not rebuild the normals or tangents.
 		bool bComputeTangentLegacy = !BuildSettings->bUseMikkTSpace && (BuildSettings->bRecomputeNormals || BuildSettings->bRecomputeTangents);
@@ -98,15 +99,19 @@ UMeshDescription* FMeshDescriptionHelper::GetRenderMeshDescription(UObject* Owne
 		{
 			if (!bHasAllNormals)
 			{
-				EComputeNTBsOptions ComputeNTBsOptions = EComputeNTBsOptions::Normals;
-				RenderMeshDescription->ComputeTangentsAndNormals(ComputeNTBsOptions);
+				FMeshDescriptionOperations::CreateNormals(RenderMeshDescription, (FMeshDescriptionOperations::ETangentOptions)TangentOptions, false);
+
+				//EComputeNTBsOptions ComputeNTBsOptions = EComputeNTBsOptions::Normals;
+				//RenderMeshDescription->ComputeTangentsAndNormals(ComputeNTBsOptions);
 			}
 			FMeshDescriptionOperations::CreateMikktTangents(RenderMeshDescription, (FMeshDescriptionOperations::ETangentOptions)TangentOptions);
 		}
 		else if(!bHasAllNormals || !bHasAllTangents)
 		{
-			EComputeNTBsOptions ComputeNTBsOptions = (bHasAllNormals ? EComputeNTBsOptions::None : EComputeNTBsOptions::Normals) | (bHasAllTangents ? EComputeNTBsOptions::None : EComputeNTBsOptions::Tangents);
-			RenderMeshDescription->ComputeTangentsAndNormals(ComputeNTBsOptions);
+			FMeshDescriptionOperations::CreateNormals(RenderMeshDescription, (FMeshDescriptionOperations::ETangentOptions)TangentOptions, true);
+
+			//EComputeNTBsOptions ComputeNTBsOptions = (bHasAllNormals ? EComputeNTBsOptions::None : EComputeNTBsOptions::Normals) | (bHasAllTangents ? EComputeNTBsOptions::None : EComputeNTBsOptions::Tangents);
+			//RenderMeshDescription->ComputeTangentsAndNormals(ComputeNTBsOptions);
 		}
 	}
 
diff --git a/Engine/Source/Developer/MeshDescriptionOperations/Private/MeshDescriptionOperations.cpp b/Engine/Source/Developer/MeshDescriptionOperations/Private/MeshDescriptionOperations.cpp
--- a/Engine/Source/Developer/MeshDescriptionOperations/Private/MeshDescriptionOperations.cpp
+++ b/Engine/Source/Developer/MeshDescriptionOperations/Private/MeshDescriptionOperations.cpp
@@ -316,18 +316,105 @@ void FMeshDescriptionOperations::ConverToRawMesh(const UMeshDescription* SourceM
 	ConvertHardEdgesToSmoothGroup(SourceMeshDescription, DestinationRawMesh);
 }
 
-void FMeshDescriptionOperations::ConverFromRawMesh(const struct FRawMesh &SourceRawMesh, UMeshDescription* DestinationMeshDescription, const TMap<int32, FName>& MaterialMap)
+//We want to fill the UMeshDescription vertex position mesh attribute with the FRawMesh vertex position
+//We will also weld the vertex position (old FRawMesh is not always welded) and construct a mapping array to match the FVertexID
+void FillMeshDescriptionVertexPositionNoDuplicate(const TArray<FVector> &RawMeshVertexPositions, UMeshDescription* DestinationMeshDescription, TArray<FVertexID>& RemapVertexPosition)
+{
+	TVertexAttributeArray<FVector>& VertexPositions = DestinationMeshDescription->VertexAttributes().GetAttributes<FVector>(MeshAttribute::Vertex::Position);
+
+	const int32 NumVertex = RawMeshVertexPositions.Num();
+
+	TMap<int32, int32> TempRemapVertexPosition;
+	TempRemapVertexPosition.Reserve(NumVertex);
+
+	// Create a list of vertex Z/index pairs
+	TArray<FIndexAndZ> VertIndexAndZ;
+	VertIndexAndZ.Reserve(NumVertex);
+
+	for (int32 VertexIndex = 0; VertexIndex < NumVertex; ++VertexIndex)
+	{
+		new(VertIndexAndZ)FIndexAndZ(VertexIndex, RawMeshVertexPositions[VertexIndex]);
+	}
+
+	// Sort the vertices by z value
+	VertIndexAndZ.Sort(FCompareIndexAndZ());
+
+	int32 VertexCount = 0;
+	// Search for duplicates, quickly!
+	for (int32 i = 0; i < VertIndexAndZ.Num(); i++)
+	{
+		int32 Index_i = VertIndexAndZ[i].Index;
+		if (TempRemapVertexPosition.Contains(Index_i))
+		{
+			continue;
+		}
+		TempRemapVertexPosition.FindOrAdd(Index_i) = VertexCount;
+		// only need to search forward, since we add pairs both ways
+		for (int32 j = i + 1; j < VertIndexAndZ.Num(); j++)
+		{
+			if (FMath::Abs(VertIndexAndZ[j].Z - VertIndexAndZ[i].Z) > SMALL_NUMBER)
+				break; // can't be any more dups
+
+			const FVector& PositionA = *(VertIndexAndZ[i].OriginalVector);
+			const FVector& PositionB = *(VertIndexAndZ[j].OriginalVector);
+
+			if (PositionA.Equals(PositionB, SMALL_NUMBER))
+			{
+				TempRemapVertexPosition.FindOrAdd(VertIndexAndZ[j].Index) = VertexCount;
+			}
+		}
+		VertexCount++;
+	}
+
+	//Make sure the vertex are added in the same order to be lossless when converting the FRawMesh
+	//In case there is a duplicate even reordering it will not be lossless, but MeshDescription do not support
+	//bad data like duplicated vertex position.
+	RemapVertexPosition.AddUninitialized(NumVertex);
+	DestinationMeshDescription->ReserveNewVertices(VertexCount);
+	TArray<FVertexID> UniqueVertexDone;
+	UniqueVertexDone.AddUninitialized(VertexCount);
+	for (int32 VertexIndex = 0; VertexIndex < VertexCount; ++VertexIndex)
+	{
+		UniqueVertexDone[VertexIndex] = FVertexID::Invalid;
+	}
+	for (int32 VertexIndex = 0; VertexIndex < NumVertex; ++VertexIndex)
+	{
+		int32 RealIndex = TempRemapVertexPosition[VertexIndex];
+		if (UniqueVertexDone[RealIndex] != FVertexID::Invalid)
+		{
+			RemapVertexPosition[VertexIndex] = UniqueVertexDone[RealIndex];
+			continue;
+		}
+		FVertexID VertexID = DestinationMeshDescription->CreateVertex();
+		UniqueVertexDone[RealIndex] = VertexID;
+		VertexPositions[VertexID] = RawMeshVertexPositions[VertexIndex];
+		RemapVertexPosition[VertexIndex] = VertexID;
+	}
+}
+
+//Discover degenerated triangle
+bool IsTriangleDegenerated(const FRawMesh &SourceRawMesh, const TArray<FVertexID>& RemapVertexPosition, const int32 VerticeIndexBase)
+{
+	FVertexID VertexIDs[3];
+	for (int32 Corner = 0; Corner < 3; ++Corner)
+	{
+		int32 VerticeIndex = VerticeIndexBase + Corner;
+		VertexIDs[Corner] = RemapVertexPosition[SourceRawMesh.WedgeIndices[VerticeIndex]];
+	}
+	return (VertexIDs[0] == VertexIDs[1] || VertexIDs[0] == VertexIDs[2] || VertexIDs[1] == VertexIDs[2]);
+}
+
+void FMeshDescriptionOperations::ConverFromRawMesh(const FRawMesh &SourceRawMesh, UMeshDescription* DestinationMeshDescription, const TMap<int32, FName>& MaterialMap)
 {
 	check(DestinationMeshDescription != nullptr);
 	DestinationMeshDescription->Empty();
 
-	DestinationMeshDescription->ReserveNewVertices(SourceRawMesh.VertexPositions.Num());
 	DestinationMeshDescription->ReserveNewVertexInstances(SourceRawMesh.WedgeIndices.Num());
 	DestinationMeshDescription->ReserveNewPolygons(SourceRawMesh.WedgeIndices.Num() / 3);
-	DestinationMeshDescription->ReserveNewEdges(SourceRawMesh.WedgeIndices.Num() * 2 / 3); //Approximately 2 edges per polygons
-																						   //Gather all array data
-	TVertexAttributeArray<FVector>& VertexPositions = DestinationMeshDescription->VertexAttributes().GetAttributes<FVector>(MeshAttribute::Vertex::Position);
+	//Approximately 2.5 edges per polygons
+	DestinationMeshDescription->ReserveNewEdges(SourceRawMesh.WedgeIndices.Num() * 2.5f / 3);
 
+	//Gather all array data
 	TVertexInstanceAttributeArray<FVector>& VertexInstanceNormals = DestinationMeshDescription->VertexInstanceAttributes().GetAttributes<FVector>(MeshAttribute::VertexInstance::Normal);
 	TVertexInstanceAttributeArray<FVector>& VertexInstanceTangents = DestinationMeshDescription->VertexInstanceAttributes().GetAttributes<FVector>(MeshAttribute::VertexInstance::Tangent);
 	TVertexInstanceAttributeArray<float>& VertexInstanceBinormalSigns = DestinationMeshDescription->VertexInstanceAttributes().GetAttributes<float>(MeshAttribute::VertexInstance::BinormalSign);
@@ -350,11 +437,10 @@ void FMeshDescriptionOperations::ConverFromRawMesh(const struct FRawMesh &Source
 		}
 	}
 	VertexInstanceUVs.SetNumIndices(NumTexCoords);
-	for (int32 VertexIndex = 0; VertexIndex < SourceRawMesh.VertexPositions.Num(); ++VertexIndex)
-	{
-		FVertexID VertexID = DestinationMeshDescription->CreateVertex();
-		VertexPositions[VertexID] = SourceRawMesh.VertexPositions[VertexIndex];
-	}
+
+	//Ensure we do not have any duplicate, We found all duplicated vertex and compact them and build a remap indice array to remap the wedgeindices
+	TArray<FVertexID> RemapVertexPosition;
+	FillMeshDescriptionVertexPositionNoDuplicate(SourceRawMesh.VertexPositions, DestinationMeshDescription, RemapVertexPosition);
 
 	bool bHasColors = SourceRawMesh.WedgeColors.Num() > 0;
 	bool bHasTangents = SourceRawMesh.WedgeTangentX.Num() > 0 && SourceRawMesh.WedgeTangentY.Num() > 0;
@@ -367,7 +453,12 @@ void FMeshDescriptionOperations::ConverFromRawMesh(const struct FRawMesh &Source
 	for (int32 TriangleIndex = 0; TriangleIndex < TriangleCount; ++TriangleIndex)
 	{
 		int32 VerticeIndexBase = TriangleIndex * 3;
-
+		//Check if the triangle is degenerated and skip the data if its the case
+		if (IsTriangleDegenerated(SourceRawMesh, RemapVertexPosition, VerticeIndexBase))
+		{
+			continue;
+		}
+		
 		//PolygonGroup
 		FPolygonGroupID PolygonGroupID = FPolygonGroupID::Invalid;
 		FName PolygonGroupImportedMaterialSlotName = NAME_None;
@@ -391,12 +482,13 @@ void FMeshDescriptionOperations::ConverFromRawMesh(const struct FRawMesh &Source
 			PolygonGroupImportedMaterialSlotNames[PolygonGroupID] = PolygonGroupImportedMaterialSlotName == NAME_None ? FName(*FString::Printf(TEXT("MaterialSlot_%d"), MaterialIndex)) : PolygonGroupImportedMaterialSlotName;
 			PolygonGroups.Add(PolygonGroupID);
 		}
-
+		FVertexInstanceID TriangleVertexInstanceIDs[3];
 		for (int32 Corner = 0; Corner < 3; ++Corner)
 		{
 			int32 VerticeIndex = VerticeIndexBase + Corner;
-			FVertexID VertexID(SourceRawMesh.WedgeIndices[VerticeIndex]);
+			FVertexID VertexID = RemapVertexPosition[SourceRawMesh.WedgeIndices[VerticeIndex]];
 			FVertexInstanceID VertexInstanceID = DestinationMeshDescription->CreateVertexInstance(VertexID);
+			TriangleVertexInstanceIDs[Corner] = VertexInstanceID;
 			VertexInstanceColors[VertexInstanceID] = bHasColors ? FLinearColor::FromSRGBColor(SourceRawMesh.WedgeColors[VerticeIndex]) : FLinearColor::White;
 			VertexInstanceTangents[VertexInstanceID] = bHasTangents ? SourceRawMesh.WedgeTangentX[VerticeIndex] : FVector(ForceInitToZero);
 			VertexInstanceBinormalSigns[VertexInstanceID] = bHasTangents ? GetBasisDeterminantSign(SourceRawMesh.WedgeTangentX[VerticeIndex].GetSafeNormal(), SourceRawMesh.WedgeTangentY[VerticeIndex].GetSafeNormal(), SourceRawMesh.WedgeTangentZ[VerticeIndex].GetSafeNormal()) : 0.0f;
@@ -415,7 +507,6 @@ void FMeshDescriptionOperations::ConverFromRawMesh(const struct FRawMesh &Source
 
 		//Create the polygon edges
 		TArray<UMeshDescription::FContourPoint> Contours;
-		// Add the edges of this triangle
 		for (uint32 Corner = 0; Corner < 3; ++Corner)
 		{
 			int32 ContourPointIndex = Contours.AddDefaulted();
@@ -426,8 +517,8 @@ void FMeshDescriptionOperations::ConverFromRawMesh(const struct FRawMesh &Source
 			CornerIndices[1] = (Corner + 1) % 3;
 
 			FVertexID EdgeVertexIDs[2];
-			EdgeVertexIDs[0] = DestinationMeshDescription->GetVertexInstanceVertex(FVertexInstanceID(VerticeIndexBase + CornerIndices[0]));
-			EdgeVertexIDs[1] = DestinationMeshDescription->GetVertexInstanceVertex(FVertexInstanceID(VerticeIndexBase + CornerIndices[1]));
+			EdgeVertexIDs[0] = DestinationMeshDescription->GetVertexInstanceVertex(FVertexInstanceID(TriangleVertexInstanceIDs[CornerIndices[0]]));
+			EdgeVertexIDs[1] = DestinationMeshDescription->GetVertexInstanceVertex(FVertexInstanceID(TriangleVertexInstanceIDs[CornerIndices[1]]));
 
 			FEdgeID MatchEdgeId = DestinationMeshDescription->GetVertexPairEdge(EdgeVertexIDs[0], EdgeVertexIDs[1]);
 			if (MatchEdgeId == FEdgeID::Invalid)
@@ -435,7 +526,7 @@ void FMeshDescriptionOperations::ConverFromRawMesh(const struct FRawMesh &Source
 				MatchEdgeId = DestinationMeshDescription->CreateEdge(EdgeVertexIDs[0], EdgeVertexIDs[1]);
 			}
 			ContourPoint.EdgeID = MatchEdgeId;
-			ContourPoint.VertexInstanceID = FVertexInstanceID(VerticeIndexBase + CornerIndices[0]);
+			ContourPoint.VertexInstanceID = FVertexInstanceID(TriangleVertexInstanceIDs[CornerIndices[0]]);
 		}
 
 		const FPolygonID NewPolygonID = DestinationMeshDescription->CreatePolygon(PolygonGroupID, Contours);
@@ -443,20 +534,27 @@ void FMeshDescriptionOperations::ConverFromRawMesh(const struct FRawMesh &Source
 		FMeshTriangle& NewTriangle = DestinationMeshDescription->GetPolygonTriangles(NewPolygonID)[NewTriangleIndex];
 		for (int32 Corner = 0; Corner < 3; ++Corner)
 		{
-			FVertexInstanceID VertexInstanceID(VerticeIndexBase + Corner);
+			FVertexInstanceID VertexInstanceID = TriangleVertexInstanceIDs[Corner];
 			NewTriangle.SetVertexInstanceID(Corner, VertexInstanceID);
 		}
 	}
-	DestinationMeshDescription->ComputePolygonTangentsAndNormals(0.0f);
+	//DestinationMeshDescription->ComputePolygonTangentsAndNormals(0.0f);
+	FMeshDescriptionOperations::CreatePolygonNTB(DestinationMeshDescription, 0.0f);
+
+	ConvertSmoothGroupToHardEdges(SourceRawMesh, DestinationMeshDescription);
 
 	//Create the missing normals and tangents, should we use Mikkt space for tangent???
 	if (!bHasNormals || !bHasTangents)
 	{
-		EComputeNTBsOptions ComputeNTBsOptions = (bHasNormals ? EComputeNTBsOptions::None : EComputeNTBsOptions::Normals) | (bHasTangents ? EComputeNTBsOptions::None : EComputeNTBsOptions::Tangents);
-		DestinationMeshDescription->ComputeTangentsAndNormals(ComputeNTBsOptions);
+		//EComputeNTBsOptions ComputeNTBsOptions = (bHasNormals ? EComputeNTBsOptions::None : EComputeNTBsOptions::Normals) | (bHasTangents ? EComputeNTBsOptions::None : EComputeNTBsOptions::Tangents);
+		//DestinationMeshDescription->ComputeTangentsAndNormals(ComputeNTBsOptions);
+		//Create the missing normals and tangents
+		if (!bHasNormals)
+		{
+			CreateNormals(DestinationMeshDescription, ETangentOptions::BlendOverlappingNormals, false);
+		}
+		CreateMikktTangents(DestinationMeshDescription, ETangentOptions::BlendOverlappingNormals);
 	}
-
-	ConvertSmoothGroupToHardEdges(SourceRawMesh, DestinationMeshDescription);
 }
 
 //////////////////////////////////////////////////////////////////////////
diff --git a/Engine/Source/Runtime/Engine/Private/StaticMesh.cpp b/Engine/Source/Runtime/Engine/Private/StaticMesh.cpp
--- a/Engine/Source/Runtime/Engine/Private/StaticMesh.cpp
+++ b/Engine/Source/Runtime/Engine/Private/StaticMesh.cpp
@@ -1314,7 +1314,7 @@ FArchive& operator<<(FArchive& Ar, FMeshBuildSettings& BuildSettings)
 // differences, etc.) replace the version GUID below with a new one.
 // In case of merge conflicts with DDC versions, you MUST generate a new GUID
 // and set this new GUID as the version.                                       
-#define STATICMESH_DERIVEDDATA_VER TEXT("D6FDB3200D904958970912FE16AAA185")
+#define STATICMESH_DERIVEDDATA_VER TEXT("23E9164C89904B29B9A209F087DF9C0D")
 
 static const FString& GetStaticMeshDerivedDataVersion()
 {
@@ -2683,7 +2683,7 @@ void UStaticMesh::FixupMaterialSlotName()
 // differences, etc.) replace the version GUID below with a new one.
 // In case of merge conflicts with DDC versions, you MUST generate a new GUID
 // and set this new GUID as the version.                                       
-#define MESHDATAKEY_STATICMESH_DERIVEDDATA_VER TEXT("8ABEA51A53B84861846CAE42E9C11538")
+#define MESHDATAKEY_STATICMESH_DERIVEDDATA_VER TEXT("1536BB4EBA6845F5AECFDA4C12FA216B")
 
 static const FString& GetMeshDataKeyStaticMeshDerivedDataVersion()
 {
