Fix linux crash, there was two structure with the same name

diff --git a/Engine/Source/Developer/MeshDescriptionOperations/Private/MeshDescriptionOperations.cpp b/Engine/Source/Developer/MeshDescriptionOperations/Private/MeshDescriptionOperations.cpp
--- a/Engine/Source/Developer/MeshDescriptionOperations/Private/MeshDescriptionOperations.cpp
+++ b/Engine/Source/Developer/MeshDescriptionOperations/Private/MeshDescriptionOperations.cpp
@@ -32,29 +32,32 @@ struct FVertexInfo
 };
 
 /** Helper struct for building acceleration structures. */
-struct FIndexAndZ
+namespace MeshDescriptionOperationNamespace
 {
-	float Z;
-	int32 Index;
-	const FVector *OriginalVector;
+	struct FIndexAndZ
+	{
+		float Z;
+		int32 Index;
+		const FVector *OriginalVector;
 
-	/** Default constructor. */
-	FIndexAndZ() {}
+		/** Default constructor. */
+		FIndexAndZ() {}
 
-	/** Initialization constructor. */
-	FIndexAndZ(int32 InIndex, const FVector& V)
+		/** Initialization constructor. */
+		FIndexAndZ(int32 InIndex, const FVector& V)
+		{
+			Z = 0.30f * V.X + 0.33f * V.Y + 0.37f * V.Z;
+			Index = InIndex;
+			OriginalVector = &V;
+		}
+	};
+	/** Sorting function for vertex Z/index pairs. */
+	struct FCompareIndexAndZ
 	{
-		Z = 0.30f * V.X + 0.33f * V.Y + 0.37f * V.Z;
-		Index = InIndex;
-		OriginalVector = &V;
-	}
-};
+		FORCEINLINE bool operator()(FIndexAndZ const& A, FIndexAndZ const& B) const { return A.Z < B.Z; }
+	};
+}
 
-/** Sorting function for vertex Z/index pairs. */
-struct FCompareIndexAndZ
-{
-	FORCEINLINE bool operator()(FIndexAndZ const& A, FIndexAndZ const& B) const { return A.Z < B.Z; }
-};
 
 //////////////////////////////////////////////////////////////////////////
 // Converters
@@ -328,16 +331,16 @@ void FillMeshDescriptionVertexPositionNoDuplicate(const TArray<FVector> &RawMesh
 	TempRemapVertexPosition.Reserve(NumVertex);
 
 	// Create a list of vertex Z/index pairs
-	TArray<FIndexAndZ> VertIndexAndZ;
+	TArray<MeshDescriptionOperationNamespace::FIndexAndZ> VertIndexAndZ;
 	VertIndexAndZ.Reserve(NumVertex);
 
 	for (int32 VertexIndex = 0; VertexIndex < NumVertex; ++VertexIndex)
 	{
-		new(VertIndexAndZ)FIndexAndZ(VertexIndex, RawMeshVertexPositions[VertexIndex]);
+		new(VertIndexAndZ)MeshDescriptionOperationNamespace::FIndexAndZ(VertexIndex, RawMeshVertexPositions[VertexIndex]);
 	}
 
 	// Sort the vertices by z value
-	VertIndexAndZ.Sort(FCompareIndexAndZ());
+	VertIndexAndZ.Sort(MeshDescriptionOperationNamespace::FCompareIndexAndZ());
 
 	int32 VertexCount = 0;
 	// Search for duplicates, quickly!
@@ -1011,18 +1014,18 @@ void FMeshDescriptionOperations::FindOverlappingCorners(TMultiMap<int32, int32>&
 	const int32 NumWedges = VertexInstanceArray.Num();
 
 	// Create a list of vertex Z/index pairs
-	TArray<FIndexAndZ> VertIndexAndZ;
+	TArray<MeshDescriptionOperationNamespace::FIndexAndZ> VertIndexAndZ;
 	VertIndexAndZ.Reserve(NumWedges);
 
 	const TVertexAttributeArray<FVector>& VertexPositions = MeshDescription->VertexAttributes().GetAttributes<FVector>(MeshAttribute::Vertex::Position);
 
 	for (const FVertexInstanceID VertexInstanceID : VertexInstanceArray.GetElementIDs())
 	{
-		new(VertIndexAndZ)FIndexAndZ(VertexInstanceID.GetValue(), VertexPositions[MeshDescription->GetVertexInstanceVertex(VertexInstanceID)]);
+		new(VertIndexAndZ)MeshDescriptionOperationNamespace::FIndexAndZ(VertexInstanceID.GetValue(), VertexPositions[MeshDescription->GetVertexInstanceVertex(VertexInstanceID)]);
 	}
 
 	// Sort the vertices by z value
-	VertIndexAndZ.Sort(FCompareIndexAndZ());
+	VertIndexAndZ.Sort(MeshDescriptionOperationNamespace::FCompareIndexAndZ());
 
 	// Search for duplicates, quickly!
 	for (int32 i = 0; i < VertIndexAndZ.Num(); i++)
