Fix possible crash when a mesh has very small triangles, we do not remove degenerate polygons anymore when computing polygon normals. Degenerate poygons are remove later in the process.

diff --git a/Engine/Source/Developer/MeshDescriptionOperations/Private/MeshDescriptionOperations.cpp b/Engine/Source/Developer/MeshDescriptionOperations/Private/MeshDescriptionOperations.cpp
--- a/Engine/Source/Developer/MeshDescriptionOperations/Private/MeshDescriptionOperations.cpp
+++ b/Engine/Source/Developer/MeshDescriptionOperations/Private/MeshDescriptionOperations.cpp
@@ -564,14 +564,15 @@ void FMeshDescriptionOperations::ConverFromRawMesh(const FRawMesh &SourceRawMesh
 			NewTriangle.SetVertexInstanceID(Corner, VertexInstanceID);
 		}
 	}
-	//DestinationMeshDescription->ComputePolygonTangentsAndNormals(0.0f);
-	FMeshDescriptionOperations::CreatePolygonNTB(DestinationMeshDescription, 0.0f);
-
+	
 	ConvertSmoothGroupToHardEdges(SourceRawMesh, DestinationMeshDescription);
 
 	//Create the missing normals and tangents, should we use Mikkt space for tangent???
 	if (!bHasNormals || !bHasTangents)
 	{
+		//DestinationMeshDescription->ComputePolygonTangentsAndNormals(0.0f);
+		FMeshDescriptionOperations::CreatePolygonNTB(DestinationMeshDescription, 0.0f);
+
 		//EComputeNTBsOptions ComputeNTBsOptions = (bHasNormals ? EComputeNTBsOptions::None : EComputeNTBsOptions::Normals) | (bHasTangents ? EComputeNTBsOptions::None : EComputeNTBsOptions::Tangents);
 		//DestinationMeshDescription->ComputeTangentsAndNormals(ComputeNTBsOptions);
 		//Create the missing normals and tangents
@@ -594,8 +595,6 @@ void FMeshDescriptionOperations::CreatePolygonNTB(UMeshDescription* MeshDescript
 	TPolygonAttributeArray<FVector>& PolygonTangents = MeshDescription->PolygonAttributes().GetAttributes<FVector>(MeshAttribute::Polygon::Tangent);
 	TPolygonAttributeArray<FVector>& PolygonBinormals = MeshDescription->PolygonAttributes().GetAttributes<FVector>(MeshAttribute::Polygon::Binormal);
 
-	TArray<FPolygonID> DegeneratePolygons;
-
 	FVertexInstanceArray& VertexInstanceArray = MeshDescription->VertexInstances();
 	FVertexArray& VertexArray = MeshDescription->Vertices();
 	FPolygonArray& PolygonArray = MeshDescription->Polygons();
@@ -659,50 +658,19 @@ void FMeshDescriptionOperations::CreatePolygonNTB(UMeshDescription* MeshDescript
 			}
 			else
 			{
-				DegeneratePolygons.AddUnique(PolygonID);
+				//This will force a recompute of the normals and tangents
+				TangentX = FVector(0.0f);
+				TangentY = FVector(0.0f);
+				TangentZ = FVector(0.0f);
+				break;
 			}
 		}
-
 		TangentX.Normalize();
 		TangentY.Normalize();
 		TangentZ.Normalize();
 		PolygonTangents[PolygonID] = TangentX;
 		PolygonBinormals[PolygonID] = TangentY;
 		PolygonNormals[PolygonID] = TangentZ;
-		
-	}
-
-	//Delete the degenerated polygons. The array is fill only if the remove degenerated option is turn on.
-	if (DegeneratePolygons.Num() > 0)
-	{
-		TArray<FEdgeID> OrphanedEdges;
-		TArray<FVertexInstanceID> OrphanedVertexInstances;
-		TArray<FPolygonGroupID> OrphanedPolygonGroups;
-		TArray<FVertexID> OrphanedVertices;
-		for (FPolygonID& PolygonID : DegeneratePolygons)
-		{
-			MeshDescription->DeletePolygon(PolygonID, &OrphanedEdges, &OrphanedVertexInstances, &OrphanedPolygonGroups);
-		}
-		for (FPolygonGroupID& PolygonGroupID : OrphanedPolygonGroups)
-		{
-			MeshDescription->DeletePolygonGroup(PolygonGroupID);
-		}
-		for (FVertexInstanceID& VertexInstanceID : OrphanedVertexInstances)
-		{
-			MeshDescription->DeleteVertexInstance(VertexInstanceID, &OrphanedVertices);
-		}
-		for (FEdgeID& EdgeID : OrphanedEdges)
-		{
-			MeshDescription->DeleteEdge(EdgeID, &OrphanedVertices);
-		}
-		for (FVertexID& VertexID : OrphanedVertices)
-		{
-			MeshDescription->DeleteVertex(VertexID);
-		}
-		//Compact and Remap IDs so we have clean ID from 0 to n since we just erase some polygons
-		//The render build need to have compact ID
-		FElementIDRemappings RemappingInfos;
-		MeshDescription->Compact(RemappingInfos);
 	}
 }
 
