Fix CIS build for monolitic build
Fix the old faceted asset
Bump the staticmesh version

diff --git a/Engine/Source/Developer/MeshBuilder/Private/BuildOptimizationHelper.h b/Engine/Source/Developer/MeshBuilder/Private/BuildOptimizationHelper.h
--- a/Engine/Source/Developer/MeshBuilder/Private/BuildOptimizationHelper.h
+++ b/Engine/Source/Developer/MeshBuilder/Private/BuildOptimizationHelper.h
@@ -285,63 +285,6 @@ public:
 		}
 	}
 
-	static void BuildStaticAdjacencyIndexBuffer(
-		const FPositionVertexBuffer& PositionVertexBuffer,
-		const FStaticMeshVertexBuffer& VertexBuffer,
-		const TArray<uint32>& Indices,
-		TArray<uint32>& OutPnAenIndices
-	)
-	{
-		if (Indices.Num())
-		{
-			BuildOptimizationPrivate::NvTriStripHelper::FStaticMeshNvRenderBuffer StaticMeshRenderBuffer(PositionVertexBuffer, VertexBuffer, Indices);
-			nv::IndexBuffer* PnAENIndexBuffer = nv::tess::buildTessellationBuffer(&StaticMeshRenderBuffer, nv::DBM_PnAenDominantCorner, true);
-			check(PnAENIndexBuffer);
-			const int32 IndexCount = (int32)PnAENIndexBuffer->getLength();
-			OutPnAenIndices.Empty(IndexCount);
-			OutPnAenIndices.AddUninitialized(IndexCount);
-			for (int32 Index = 0; Index < IndexCount; ++Index)
-			{
-				OutPnAenIndices[Index] = (*PnAENIndexBuffer)[Index];
-			}
-			delete PnAENIndexBuffer;
-		}
-		else
-		{
-			OutPnAenIndices.Empty();
-		}
-	}
-
-	static void CacheOptimizeIndexBuffer(TArray<uint16>& Indices)
-	{
-		bool bDisableTriangleOrderOptimization = (BuildOptimizationPrivate::CVarTriangleOrderOptimization.GetValueOnGameThread() == 2);
-		bool bUsingNvTriStrip = !bDisableTriangleOrderOptimization && (BuildOptimizationPrivate::CVarTriangleOrderOptimization.GetValueOnGameThread() == 0);
-
-		if (bUsingNvTriStrip)
-		{
-			BuildOptimizationPrivate::NvTriStripHelper::CacheOptimizeIndexBuffer(Indices);
-		}
-		else if (!bDisableTriangleOrderOptimization)
-		{
-			BuildOptimizationPrivate::ForsythHelper::CacheOptimizeIndexBuffer(Indices);
-		}
-	}
-
-	static void CacheOptimizeIndexBuffer(TArray<uint32>& Indices)
-	{
-		bool bDisableTriangleOrderOptimization = (BuildOptimizationPrivate::CVarTriangleOrderOptimization.GetValueOnGameThread() == 2);
-		bool bUsingNvTriStrip = !bDisableTriangleOrderOptimization && (BuildOptimizationPrivate::CVarTriangleOrderOptimization.GetValueOnGameThread() == 0);
-
-		if (bUsingNvTriStrip)
-		{
-			BuildOptimizationPrivate::NvTriStripHelper::CacheOptimizeIndexBuffer(Indices);
-		}
-		else if (!bDisableTriangleOrderOptimization)
-		{
-			BuildOptimizationPrivate::ForsythHelper::CacheOptimizeIndexBuffer(Indices);
-		}
-	}
-
 	static void CacheOptimizeVertexAndIndexBuffer(
 		TArray<FStaticMeshBuildVertex>& Vertices,
 		TArray<TArray<uint32> >& PerSectionIndices,
@@ -367,7 +310,7 @@ public:
 			if (Indices.Num())
 			{
 				// Optimize the index buffer for the post transform cache with.
-				CacheOptimizeIndexBuffer(Indices);
+				BuildOptimizationThirdParty::CacheOptimizeIndexBuffer(Indices);
 
 				// Copy the index buffer since we will be reordering it
 				TArray<uint32> OriginalIndices = Indices;
diff --git a/Engine/Source/Developer/MeshBuilder/Private/StaticMeshBuilder.cpp b/Engine/Source/Developer/MeshBuilder/Private/StaticMeshBuilder.cpp
--- a/Engine/Source/Developer/MeshBuilder/Private/StaticMeshBuilder.cpp
+++ b/Engine/Source/Developer/MeshBuilder/Private/StaticMeshBuilder.cpp
@@ -494,7 +494,7 @@ void BuildAllBufferOptimizations(FStaticMeshLODResources& StaticMeshLOD, const F
 
 		if (DepthOnlyIndices.Num() < 50000 * 3)
 		{
-			BuildOptimizationHelper::CacheOptimizeIndexBuffer(DepthOnlyIndices);
+			BuildOptimizationThirdParty::CacheOptimizeIndexBuffer(DepthOnlyIndices);
 		}
 
 		StaticMeshLOD.DepthOnlyIndexBuffer.SetIndices(DepthOnlyIndices, IndexBufferStride);
@@ -534,7 +534,7 @@ void BuildAllBufferOptimizations(FStaticMeshLODResources& StaticMeshLOD, const F
 	{
 		TArray<uint32> AdjacencyIndices;
 
-		BuildOptimizationHelper::BuildStaticAdjacencyIndexBuffer(
+		BuildOptimizationThirdParty::NvTriStripHelper::BuildStaticAdjacencyIndexBuffer(
 			StaticMeshLOD.VertexBuffers.PositionVertexBuffer,
 			StaticMeshLOD.VertexBuffers.StaticMeshVertexBuffer,
 			IndexBuffer,
diff --git a/Engine/Source/Developer/MeshBuilder/Private/ThirdPartyBuildOptimizationHelper.cpp b/Engine/Source/Developer/MeshBuilder/Private/ThirdPartyBuildOptimizationHelper.cpp
new file mode 100644
--- /dev/null
+++ b/Engine/Source/Developer/MeshBuilder/Private/ThirdPartyBuildOptimizationHelper.cpp
@@ -0,0 +1,365 @@
+// Copyright 1998-2017 Epic Games, Inc. All Rights Reserved.
+#include "ThirdPartyBuildOptimizationHelper.h"
+#include "CoreMinimal.h"
+#include "CoreTypes.h"
+
+#include "NvTriStrip.h"
+#include "forsythtriangleorderoptimizer.h"
+#include "nvtess.h"
+
+namespace BuildOptimizationThirdParty
+{
+	// CVars
+	static TAutoConsoleVariable<int32> CVarTriangleOrderOptimization(
+		TEXT("r.TriangleOrderOptimization"),
+		1,
+		TEXT("Controls the algorithm to use when optimizing the triangle order for the post-transform cache.\n")
+		TEXT("0: Use NVTriStrip (slower)\n")
+		TEXT("1: Use Forsyth algorithm (fastest)(default)")
+		TEXT("2: No triangle order optimization. (least efficient, debugging purposes only)"),
+		ECVF_Default);
+
+	namespace NvTriStripHelper
+	{
+		/**
+		* Converts 16 bit indices to 32 bit prior to passing them into the real GenerateStrips util method
+		*/
+		void GenerateStrips(
+			const uint8* Indices,
+			bool Is32Bit,
+			const uint32 NumIndices,
+			PrimitiveGroup** PrimGroups,
+			uint32* NumGroups
+		)
+		{
+			if (Is32Bit)
+			{
+				GenerateStrips((uint32*)Indices, NumIndices, PrimGroups, NumGroups);
+			}
+			else
+			{
+				// convert to 32 bit
+				uint32 Idx;
+				TArray<uint32> NewIndices;
+				NewIndices.AddUninitialized(NumIndices);
+				for (Idx = 0; Idx < NumIndices; ++Idx)
+				{
+					NewIndices[Idx] = ((uint16*)Indices)[Idx];
+				}
+				GenerateStrips(NewIndices.GetData(), NumIndices, PrimGroups, NumGroups);
+			}
+
+		}
+
+		/**
+		* Orders a triangle list for better vertex cache coherency.
+		*
+		* *** WARNING: This is safe to call for multiple threads IF AND ONLY IF all
+		* threads call SetListsOnly(true) and SetCacheSize(CACHESIZE_GEFORCE3). If
+		* NvTriStrip is ever used with different settings the library will need
+		* some modifications to be thread-safe. ***
+		*/
+		template<typename IndexDataType, typename Allocator>
+		void CacheOptimizeIndexBuffer(TArray<IndexDataType, Allocator>& Indices)
+		{
+			static_assert(sizeof(IndexDataType) == 2 || sizeof(IndexDataType) == 4, "Indices must be short or int.");
+
+			PrimitiveGroup*	PrimitiveGroups = NULL;
+			uint32			NumPrimitiveGroups = 0;
+			bool Is32Bit = sizeof(IndexDataType) == 4;
+
+			SetListsOnly(true);
+			SetCacheSize(CACHESIZE_GEFORCE3);
+
+			GenerateStrips((uint8*)Indices.GetData(), Is32Bit, Indices.Num(), &PrimitiveGroups, &NumPrimitiveGroups);
+
+			Indices.Empty();
+			Indices.AddUninitialized(PrimitiveGroups->numIndices);
+
+			if (Is32Bit)
+			{
+				FMemory::Memcpy(Indices.GetData(), PrimitiveGroups->indices, Indices.Num() * sizeof(IndexDataType));
+			}
+			else
+			{
+				for (uint32 I = 0; I < PrimitiveGroups->numIndices; ++I)
+				{
+					Indices[I] = (uint16)PrimitiveGroups->indices[I];
+				}
+			}
+
+			delete[] PrimitiveGroups;
+		}
+
+		/**
+		* Provides skeletal mesh render data to the NVIDIA tessellation library.
+		*/
+		class FSkeletalMeshNvRenderBuffer : public nv::RenderBuffer
+		{
+		public:
+
+			/** Construct from static mesh render buffers. */
+			FSkeletalMeshNvRenderBuffer(
+				const TArray<FSoftSkinVertex>& InVertexBuffer,
+				const uint32 InTexCoordCount,
+				const TArray<uint32>& Indices)
+				: VertexBuffer(InVertexBuffer)
+				, TexCoordCount(InTexCoordCount)
+			{
+				mIb = new nv::IndexBuffer((void*)Indices.GetData(), nv::IBT_U32, Indices.Num(), false);
+			}
+
+			/** Retrieve the position and first texture coordinate of the specified index. */
+			virtual nv::Vertex getVertex(unsigned int Index) const
+			{
+				nv::Vertex Vertex;
+
+				check(Index < (unsigned int)VertexBuffer.Num());
+
+				const FSoftSkinVertex& SrcVertex = VertexBuffer[Index];
+
+				Vertex.pos.x = SrcVertex.Position.X;
+				Vertex.pos.y = SrcVertex.Position.Y;
+				Vertex.pos.z = SrcVertex.Position.Z;
+
+				if (TexCoordCount > 0)
+				{
+					Vertex.uv.x = SrcVertex.UVs[0].X;
+					Vertex.uv.y = SrcVertex.UVs[0].Y;
+				}
+				else
+				{
+					Vertex.uv.x = 0.0f;
+					Vertex.uv.y = 0.0f;
+				}
+
+				return Vertex;
+			}
+
+		private:
+			/** The vertex buffer for the skeletal mesh. */
+			const TArray<FSoftSkinVertex>& VertexBuffer;
+			const uint32 TexCoordCount;
+
+			/** Copying is forbidden. */
+			FSkeletalMeshNvRenderBuffer(const FSkeletalMeshNvRenderBuffer&);
+			FSkeletalMeshNvRenderBuffer& operator=(const FSkeletalMeshNvRenderBuffer&);
+		};
+
+		/**
+		* Provides static mesh render data to the NVIDIA tessellation library.
+		*/
+		class FStaticMeshNvRenderBuffer : public nv::RenderBuffer
+		{
+		public:
+
+			/** Construct from static mesh render buffers. */
+			FStaticMeshNvRenderBuffer(
+				const FPositionVertexBuffer& InPositionVertexBuffer,
+				const FStaticMeshVertexBuffer& InVertexBuffer,
+				const TArray<uint32>& Indices);
+
+			/** Retrieve the position and first texture coordinate of the specified index. */
+			virtual nv::Vertex getVertex(unsigned int Index) const;
+
+		private:
+
+			/** The position vertex buffer for the static mesh. */
+			const FPositionVertexBuffer& PositionVertexBuffer;
+
+			/** The vertex buffer for the static mesh. */
+			const FStaticMeshVertexBuffer& VertexBuffer;
+
+			/** Copying is forbidden. */
+			FStaticMeshNvRenderBuffer(const FStaticMeshNvRenderBuffer&);
+			FStaticMeshNvRenderBuffer& operator=(const FStaticMeshNvRenderBuffer&);
+		};
+
+		FStaticMeshNvRenderBuffer::FStaticMeshNvRenderBuffer(
+			const FPositionVertexBuffer& InPositionVertexBuffer,
+			const FStaticMeshVertexBuffer& InVertexBuffer,
+			const TArray<uint32>& Indices)
+			: PositionVertexBuffer(InPositionVertexBuffer)
+			, VertexBuffer(InVertexBuffer)
+		{
+			check(PositionVertexBuffer.GetNumVertices() == VertexBuffer.GetNumVertices());
+			mIb = new nv::IndexBuffer((void*)Indices.GetData(), nv::IBT_U32, Indices.Num(), false);
+		}
+
+		nv::Vertex FStaticMeshNvRenderBuffer::getVertex(unsigned int Index) const
+		{
+			nv::Vertex Vertex;
+
+			check(Index < PositionVertexBuffer.GetNumVertices());
+
+			const FVector& Position = PositionVertexBuffer.VertexPosition(Index);
+			Vertex.pos.x = Position.X;
+			Vertex.pos.y = Position.Y;
+			Vertex.pos.z = Position.Z;
+
+			if (VertexBuffer.GetNumTexCoords())
+			{
+				const FVector2D UV = VertexBuffer.GetVertexUV(Index, 0);
+				Vertex.uv.x = UV.X;
+				Vertex.uv.y = UV.Y;
+			}
+			else
+			{
+				Vertex.uv.x = 0.0f;
+				Vertex.uv.y = 0.0f;
+			}
+
+			return Vertex;
+		}
+
+		void BuildStaticAdjacencyIndexBuffer(
+			const FPositionVertexBuffer& PositionVertexBuffer,
+			const FStaticMeshVertexBuffer& VertexBuffer,
+			const TArray<uint32>& Indices,
+			TArray<uint32>& OutPnAenIndices
+		)
+		{
+			if (Indices.Num())
+			{
+				FStaticMeshNvRenderBuffer StaticMeshRenderBuffer(PositionVertexBuffer, VertexBuffer, Indices);
+				nv::IndexBuffer* PnAENIndexBuffer = nv::tess::buildTessellationBuffer(&StaticMeshRenderBuffer, nv::DBM_PnAenDominantCorner, true);
+				check(PnAENIndexBuffer);
+				const int32 IndexCount = (int32)PnAENIndexBuffer->getLength();
+				OutPnAenIndices.Empty(IndexCount);
+				OutPnAenIndices.AddUninitialized(IndexCount);
+				for (int32 Index = 0; Index < IndexCount; ++Index)
+				{
+					OutPnAenIndices[Index] = (*PnAENIndexBuffer)[Index];
+				}
+				delete PnAENIndexBuffer;
+			}
+			else
+			{
+				OutPnAenIndices.Empty();
+			}
+		}
+
+		void BuildSkeletalAdjacencyIndexBuffer(
+			const TArray<FSoftSkinVertex>& VertexBuffer,
+			const uint32 TexCoordCount,
+			const TArray<uint32>& Indices,
+			TArray<uint32>& OutPnAenIndices
+		)
+		{
+			if (Indices.Num())
+			{
+				FSkeletalMeshNvRenderBuffer SkeletalMeshRenderBuffer(VertexBuffer, TexCoordCount, Indices);
+				nv::IndexBuffer* PnAENIndexBuffer = nv::tess::buildTessellationBuffer(&SkeletalMeshRenderBuffer, nv::DBM_PnAenDominantCorner, true);
+				check(PnAENIndexBuffer);
+				const int32 IndexCount = (int32)PnAENIndexBuffer->getLength();
+				OutPnAenIndices.Empty(IndexCount);
+				OutPnAenIndices.AddUninitialized(IndexCount);
+				for (int32 Index = 0; Index < IndexCount; ++Index)
+				{
+					OutPnAenIndices[Index] = (*PnAENIndexBuffer)[Index];
+				}
+				delete PnAENIndexBuffer;
+			}
+			else
+			{
+				OutPnAenIndices.Empty();
+			}
+		}
+	}
+
+	namespace ForsythHelper
+	{
+		void OptimizeFaces(
+			const uint8* Indices,
+			bool Is32Bit,
+			const uint32 NumIndices,
+			uint32 NumVertices,
+			uint32* OutIndices,
+			uint16 CacheSize
+		)
+		{
+			if (Is32Bit)
+			{
+				Forsyth::OptimizeFaces((uint32*)Indices, NumIndices, NumVertices, OutIndices, CacheSize);
+			}
+			else
+			{
+				// convert to 32 bit
+				uint32 Idx;
+				TArray<uint32> NewIndices;
+				NewIndices.AddUninitialized(NumIndices);
+				for (Idx = 0; Idx < NumIndices; ++Idx)
+				{
+					NewIndices[Idx] = ((uint16*)Indices)[Idx];
+				}
+				Forsyth::OptimizeFaces(NewIndices.GetData(), NumIndices, NumVertices, OutIndices, CacheSize);
+			}
+
+		}
+
+		template<typename IndexDataType, typename Allocator>
+		void CacheOptimizeIndexBuffer(TArray<IndexDataType, Allocator>& Indices)
+		{
+			static_assert(sizeof(IndexDataType) == 2 || sizeof(IndexDataType) == 4, "Indices must be short or int.");
+			bool Is32Bit = sizeof(IndexDataType) == 4;
+
+			// Count the number of vertices
+			uint32 NumVertices = 0;
+			for (int32 Index = 0; Index < Indices.Num(); ++Index)
+			{
+				if (Indices[Index] > NumVertices)
+				{
+					NumVertices = Indices[Index];
+				}
+			}
+			NumVertices += 1;
+
+			TArray<uint32> OptimizedIndices;
+			OptimizedIndices.AddUninitialized(Indices.Num());
+			uint16 CacheSize = 32;
+			OptimizeFaces((uint8*)Indices.GetData(), Is32Bit, Indices.Num(), NumVertices, OptimizedIndices.GetData(), CacheSize);
+
+			if (Is32Bit)
+			{
+				FMemory::Memcpy(Indices.GetData(), OptimizedIndices.GetData(), Indices.Num() * sizeof(IndexDataType));
+			}
+			else
+			{
+				for (int32 I = 0; I < OptimizedIndices.Num(); ++I)
+				{
+					Indices[I] = (uint16)OptimizedIndices[I];
+				}
+			}
+		}
+	}
+
+	void CacheOptimizeIndexBuffer(TArray<uint16>& Indices)
+	{
+		bool bDisableTriangleOrderOptimization = (CVarTriangleOrderOptimization.GetValueOnGameThread() == 2);
+		bool bUsingNvTriStrip = !bDisableTriangleOrderOptimization && (CVarTriangleOrderOptimization.GetValueOnGameThread() == 0);
+
+		if (bUsingNvTriStrip)
+		{
+			NvTriStripHelper::CacheOptimizeIndexBuffer(Indices);
+		}
+		else if (!bDisableTriangleOrderOptimization)
+		{
+			ForsythHelper::CacheOptimizeIndexBuffer(Indices);
+		}
+	}
+
+	void CacheOptimizeIndexBuffer(TArray<uint32>& Indices)
+	{
+		bool bDisableTriangleOrderOptimization = (CVarTriangleOrderOptimization.GetValueOnGameThread() == 2);
+		bool bUsingNvTriStrip = !bDisableTriangleOrderOptimization && (CVarTriangleOrderOptimization.GetValueOnGameThread() == 0);
+
+		if (bUsingNvTriStrip)
+		{
+			NvTriStripHelper::CacheOptimizeIndexBuffer(Indices);
+		}
+		else if (!bDisableTriangleOrderOptimization)
+		{
+			ForsythHelper::CacheOptimizeIndexBuffer(Indices);
+		}
+	}
+}
diff --git a/Engine/Source/Developer/MeshBuilder/Private/ThirdPartyBuildOptimizationHelper.h b/Engine/Source/Developer/MeshBuilder/Private/ThirdPartyBuildOptimizationHelper.h
deleted file mode 100644
diff --git a/Engine/Source/Developer/MeshBuilder/Public/ThirdPartyBuildOptimizationHelper.h b/Engine/Source/Developer/MeshBuilder/Public/ThirdPartyBuildOptimizationHelper.h
new file mode 100644
--- /dev/null
+++ b/Engine/Source/Developer/MeshBuilder/Public/ThirdPartyBuildOptimizationHelper.h
@@ -0,0 +1,39 @@
+// Copyright 1998-2017 Epic Games, Inc. All Rights Reserved.
+#pragma once
+
+#include "CoreMinimal.h"
+#include "CoreTypes.h"
+#include "HAL/IConsoleManager.h"
+#include "Rendering/PositionVertexBuffer.h"
+#include "Rendering/StaticMeshVertexBuffer.h"
+#include "Rendering/SkeletalMeshLODModel.h"
+
+namespace BuildOptimizationThirdParty
+{
+	MESHBUILDER_API void CacheOptimizeIndexBuffer(TArray<uint16>& Indices);
+	MESHBUILDER_API void CacheOptimizeIndexBuffer(TArray<uint32>& Indices);
+
+	/*------------------------------------------------------------------------------
+	NVTriStrip for cache optimizing index buffers.
+	------------------------------------------------------------------------------*/
+
+	namespace NvTriStripHelper
+	{
+		/*****************************
+		 * Skeletal mesh helpers
+		 */
+		MESHBUILDER_API void BuildStaticAdjacencyIndexBuffer(
+			const FPositionVertexBuffer& PositionVertexBuffer,
+			const FStaticMeshVertexBuffer& VertexBuffer,
+			const TArray<uint32>& Indices,
+			TArray<uint32>& OutPnAenIndices
+		);
+
+		MESHBUILDER_API void BuildSkeletalAdjacencyIndexBuffer(
+			const TArray<FSoftSkinVertex>& VertexBuffer,
+			const uint32 TexCoordCount,
+			const TArray<uint32>& Indices,
+			TArray<uint32>& OutPnAenIndices
+		);
+	}
+}
diff --git a/Engine/Source/Developer/MeshDescriptionOperations/Private/Allocator2D.cpp b/Engine/Source/Developer/MeshDescriptionOperations/Private/Allocator2D.cpp
--- a/Engine/Source/Developer/MeshDescriptionOperations/Private/Allocator2D.cpp
+++ b/Engine/Source/Developer/MeshDescriptionOperations/Private/Allocator2D.cpp
@@ -2,513 +2,516 @@
 
 #include "Allocator2D.h"
 
-FAllocator2D::FAllocator2D( uint32 InWidth, uint32 InHeight )
-	: Width( InWidth )
-	, Height( InHeight )
-	, LastRowFail( -1 )
+namespace MeshDescriptionOp
 {
-	Pitch = ( Width + 63 ) / 64;
-	// alloc +1 to avoid buffer overrun
-	Bits = (uint64*)FMemory::Malloc( (Pitch * Height + 1) * sizeof(uint64) );
+	FAllocator2D::FAllocator2D(uint32 InWidth, uint32 InHeight)
+		: Width(InWidth)
+		, Height(InHeight)
+		, LastRowFail(-1)
+	{
+		Pitch = (Width + 63) / 64;
+		// alloc +1 to avoid buffer overrun
+		Bits = (uint64*)FMemory::Malloc((Pitch * Height + 1) * sizeof(uint64));
 
-	Rows.Reserve( Height );
+		Rows.Reserve(Height);
 
-	for ( uint32 Y = 0; Y < Height; ++Y )
-	{
-		FRow Row;
-		Row.Index = Y;
+		for (uint32 Y = 0; Y < Height; ++Y)
+		{
+			FRow Row;
+			Row.Index = Y;
 
-		Rows.Add( Row );
+			Rows.Add(Row);
+		}
+
+		Clear();
 	}
 
-	Clear();
-}
+	FAllocator2D::~FAllocator2D()
+	{
+		FMemory::Free(Bits);
+	}
 
-FAllocator2D::~FAllocator2D()
-{
-	FMemory::Free( Bits );
-}
-
-FAllocator2D::FAllocator2D( const FAllocator2D& Other )
-	: Width( Other.Width )
-	, Height( Other.Height )
-	, Pitch( Other.Pitch )
-	, Rows( Other.Rows )
-	, LastRowFail( -1 )
-{
-	Bits = (uint64*)FMemory::Malloc( (Pitch * Height + 1) * sizeof(uint64) );
-	FMemory::Memcpy( Bits, Other.Bits, (Pitch * Height + 1) * sizeof(uint64) );
-}
+	FAllocator2D::FAllocator2D(const FAllocator2D& Other)
+		: Width(Other.Width)
+		, Height(Other.Height)
+		, Pitch(Other.Pitch)
+		, Rows(Other.Rows)
+		, LastRowFail(-1)
+	{
+		Bits = (uint64*)FMemory::Malloc((Pitch * Height + 1) * sizeof(uint64));
+		FMemory::Memcpy(Bits, Other.Bits, (Pitch * Height + 1) * sizeof(uint64));
+	}
 
-FAllocator2D& FAllocator2D::operator=( const FAllocator2D& Other )
-{
-	if ( Width != Other.Width || Height != Other.Height || Pitch != Other.Pitch )
+	FAllocator2D& FAllocator2D::operator=(const FAllocator2D& Other)
 	{
-		Width = Other.Width;
-		Height = Other.Height;
-		Pitch = Other.Pitch;
+		if (Width != Other.Width || Height != Other.Height || Pitch != Other.Pitch)
+		{
+			Width = Other.Width;
+			Height = Other.Height;
+			Pitch = Other.Pitch;
 
-		FMemory::Free( Bits );
-		Bits = (uint64*)FMemory::Malloc( (Pitch * Height + 1) * sizeof(uint64) );
-	}
+			FMemory::Free(Bits);
+			Bits = (uint64*)FMemory::Malloc((Pitch * Height + 1) * sizeof(uint64));
+		}
 
-	FMemory::Memcpy( Bits, Other.Bits, (Pitch * Height + 1) * sizeof(uint64) );
+		FMemory::Memcpy(Bits, Other.Bits, (Pitch * Height + 1) * sizeof(uint64));
 
-	Rows = Other.Rows;
+		Rows = Other.Rows;
 
-	return *this;
-}
+		return *this;
+	}
 
-void FAllocator2D::Clear()
-{
-	InitSegments();
+	void FAllocator2D::Clear()
+	{
+		InitSegments();
 
-	FMemory::Memzero( Bits, (Pitch * Height + 1) * sizeof(uint64) );
-}
+		FMemory::Memzero(Bits, (Pitch * Height + 1) * sizeof(uint64));
+	}
 
-bool FAllocator2D::Find( FRect& Rect )
-{
-	FRect TestRect = Rect;
-	for( TestRect.X = 0; TestRect.X <= Width - TestRect.W; TestRect.X++ )
+	bool FAllocator2D::Find(FRect& Rect)
 	{
-		for( TestRect.Y = 0; TestRect.Y <= Height - TestRect.H; TestRect.Y++ )
+		FRect TestRect = Rect;
+		for (TestRect.X = 0; TestRect.X <= Width - TestRect.W; TestRect.X++)
 		{
-			if( Test( TestRect ) )
+			for (TestRect.Y = 0; TestRect.Y <= Height - TestRect.H; TestRect.Y++)
 			{
-				Rect = TestRect;
-				return true;
+				if (Test(TestRect))
+				{
+					Rect = TestRect;
+					return true;
+				}
 			}
 		}
-	}
 
-	return false;
-}
+		return false;
+	}
 
-bool FAllocator2D::FindBitByBit( FRect& Rect, const FAllocator2D& Other )
-{
-	FRect TestRect = Rect;
-	for( TestRect.X = 0; TestRect.X <= Width - TestRect.W; TestRect.X++ )
+	bool FAllocator2D::FindBitByBit(FRect& Rect, const FAllocator2D& Other)
 	{
-		for( TestRect.Y = 0; TestRect.Y <= Height - TestRect.H; TestRect.Y++ )
+		FRect TestRect = Rect;
+		for (TestRect.X = 0; TestRect.X <= Width - TestRect.W; TestRect.X++)
 		{
-			if( Test( TestRect, Other ) )
+			for (TestRect.Y = 0; TestRect.Y <= Height - TestRect.H; TestRect.Y++)
 			{
-				Rect = TestRect;
-				return true;
+				if (Test(TestRect, Other))
+				{
+					Rect = TestRect;
+					return true;
+				}
 			}
 		}
-	}
-	
-	return false;
-}
 
-bool FAllocator2D::FindWithSegments( FRect& Rect, FRect BestRect, const FAllocator2D& Other )
-{
-	LastRowFail = -1;
-	FRect TestRect = Rect;
+		return false;
+	}
 
-	for( TestRect.Y = 0; TestRect.Y <= Height - TestRect.H; ++TestRect.Y )
+	bool FAllocator2D::FindWithSegments(FRect& Rect, FRect BestRect, const FAllocator2D& Other)
 	{
-		uint32 FailedLength = 1;
+		LastRowFail = -1;
+		FRect TestRect = Rect;
 
-		for( TestRect.X = 0; TestRect.X <= Width - TestRect.W; TestRect.X += FailedLength )
+		for (TestRect.Y = 0; TestRect.Y <= Height - TestRect.H; ++TestRect.Y)
 		{
-			if( TestRect.X + TestRect.Y * Height >= BestRect.X + BestRect.Y * Height )
-			{
-				// We've moved past the BestRect
-				return false;
-			}
+			uint32 FailedLength = 1;
 
-			if( TestAllRows( TestRect, Other, FailedLength ) )
+			for (TestRect.X = 0; TestRect.X <= Width - TestRect.W; TestRect.X += FailedLength)
 			{
-				Rect = TestRect;
-				return true;
+				if (TestRect.X + TestRect.Y * Height >= BestRect.X + BestRect.Y * Height)
+				{
+					// We've moved past the BestRect
+					return false;
+				}
+
+				if (TestAllRows(TestRect, Other, FailedLength))
+				{
+					Rect = TestRect;
+					return true;
+				}
 			}
 		}
+
+		return false;
 	}
-	
-	return false;
-}
 
-void FAllocator2D::Alloc( FRect Rect )
-{
-	for( uint32 y = Rect.Y; y < Rect.Y + Rect.H; y++ )
+	void FAllocator2D::Alloc(FRect Rect)
 	{
-		for( uint32 x = Rect.X; x < Rect.X + Rect.W; x++ )
+		for (uint32 y = Rect.Y; y < Rect.Y + Rect.H; y++)
 		{
-			SetBit( x, y );
+			for (uint32 x = Rect.X; x < Rect.X + Rect.W; x++)
+			{
+				SetBit(x, y);
+			}
 		}
 	}
-}
 
-void FAllocator2D::Alloc( FRect Rect, const FAllocator2D& Other )
-{
-	for( uint32 y = 0; y < Rect.H; y++ )
+	void FAllocator2D::Alloc(FRect Rect, const FAllocator2D& Other)
 	{
-		for( uint32 x = 0; x < Rect.W; x++ )
+		for (uint32 y = 0; y < Rect.H; y++)
 		{
-			if( Other.GetBit( x, y ) )
+			for (uint32 x = 0; x < Rect.W; x++)
 			{
-				SetBit( x + Rect.X, y + Rect.Y );
+				if (Other.GetBit(x, y))
+				{
+					SetBit(x + Rect.X, y + Rect.Y);
+				}
 			}
 		}
-	}
-
-	MergeSegments( Rect, Other );
-}
 
-bool FAllocator2D::TestAllRows( FRect Rect, const FAllocator2D& Other, uint32& FailedLength )
-{
-	FailedLength = 0.f;
-	bool bSuccess = true;
+		MergeSegments(Rect, Other);
+	}
 
-	if ( LastRowFail > 0 )
+	bool FAllocator2D::TestAllRows(FRect Rect, const FAllocator2D& Other, uint32& FailedLength)
 	{
-		const FRow& ThisRow = Rows[ Rect.Y + LastRowFail ];
-		const FRow& OtherRow = Other.Rows[ LastRowFail ];
+		FailedLength = 0.f;
+		bool bSuccess = true;
 
-		if ( !TestRow( ThisRow, OtherRow, Rect, FailedLength ) )
+		if (LastRowFail > 0)
 		{
-			return false;
+			const FRow& ThisRow = Rows[Rect.Y + LastRowFail];
+			const FRow& OtherRow = Other.Rows[LastRowFail];
+
+			if (!TestRow(ThisRow, OtherRow, Rect, FailedLength))
+			{
+				return false;
+			}
 		}
-	}
 
-	LastRowFail = -1;
+		LastRowFail = -1;
 
-	for ( uint32 y = 0; y < Rect.H; ++y )
-	{
-		const FRow& ThisRow = Rows[ Rect.Y + y ];
-		const FRow& OtherRow = Other.Rows[ y ];
+		for (uint32 y = 0; y < Rect.H; ++y)
+		{
+			const FRow& ThisRow = Rows[Rect.Y + y];
+			const FRow& OtherRow = Other.Rows[y];
 
-		uint32 CurrentFailedLength = 0;
-		bool bThisRowSuccess = TestRow( ThisRow, OtherRow, Rect, CurrentFailedLength );
-		bSuccess = bThisRowSuccess && bSuccess;
+			uint32 CurrentFailedLength = 0;
+			bool bThisRowSuccess = TestRow(ThisRow, OtherRow, Rect, CurrentFailedLength);
+			bSuccess = bThisRowSuccess && bSuccess;
 
-		if ( !bThisRowSuccess && CurrentFailedLength > FailedLength )
-		{
-			LastRowFail = y;
-			FailedLength = CurrentFailedLength;
-		}
+			if (!bThisRowSuccess && CurrentFailedLength > FailedLength)
+			{
+				LastRowFail = y;
+				FailedLength = CurrentFailedLength;
+			}
 
-		if ( FailedLength >= Width )
-		{
-			return false;
+			if (FailedLength >= Width)
+			{
+				return false;
+			}
 		}
-	}
-
-	return bSuccess;
-}
 
-bool FAllocator2D::TestRow( const FRow& ThisRow, const FRow& OtherRow, FRect Rect, uint32& FailedLength )
-{
-	if ( ThisRow.LongestSegment < OtherRow.LongestSegment )
-	{
-		FailedLength = Width;
-		return false;
+		return bSuccess;
 	}
 
-	uint32 StartFreeSegmentIndex = 0;
-
-	for ( const FSegment& OtherUsedSegment : OtherRow.UsedSegments )
+	bool FAllocator2D::TestRow(const FRow& ThisRow, const FRow& OtherRow, FRect Rect, uint32& FailedLength)
 	{
-		if ( OtherUsedSegment.StartPos >= Rect.W )
+		if (ThisRow.LongestSegment < OtherRow.LongestSegment)
 		{
-			break;
+			FailedLength = Width;
+			return false;
 		}
 
-		bool bFoundSpaceForSegment = false;
-		bool bFoundFutureSpaceForSegment = false;
-
-		const uint32 StartPos = Rect.X + OtherUsedSegment.StartPos;
-		const uint32 EndPos = Rect.X + FMath::Min( OtherUsedSegment.StartPos + OtherUsedSegment.Length, Rect.W );
+		uint32 StartFreeSegmentIndex = 0;
 
-		for ( int32 i = StartFreeSegmentIndex; i < ThisRow.FreeSegments.Num(); ++i )
+		for (const FSegment& OtherUsedSegment : OtherRow.UsedSegments)
 		{
-			const FSegment& ThisFreeSegment = ThisRow.FreeSegments[i];
-
-			if ( StartPos >= ThisFreeSegment.StartPos &&
-				EndPos <= ThisFreeSegment.StartPos + ThisFreeSegment.Length )
+			if (OtherUsedSegment.StartPos >= Rect.W)
 			{
-				StartFreeSegmentIndex = i;
-				bFoundSpaceForSegment = true;
 				break;
 			}
-			else if ( !bFoundFutureSpaceForSegment &&
-				StartPos < ThisFreeSegment.StartPos &&
-				OtherUsedSegment.Length <= ThisFreeSegment.Length )
+
+			bool bFoundSpaceForSegment = false;
+			bool bFoundFutureSpaceForSegment = false;
+
+			const uint32 StartPos = Rect.X + OtherUsedSegment.StartPos;
+			const uint32 EndPos = Rect.X + FMath::Min(OtherUsedSegment.StartPos + OtherUsedSegment.Length, Rect.W);
+
+			for (int32 i = StartFreeSegmentIndex; i < ThisRow.FreeSegments.Num(); ++i)
 			{
-				bFoundFutureSpaceForSegment = true;
-				FailedLength = ThisFreeSegment.StartPos - StartPos;
-				break;
+				const FSegment& ThisFreeSegment = ThisRow.FreeSegments[i];
+
+				if (StartPos >= ThisFreeSegment.StartPos &&
+					EndPos <= ThisFreeSegment.StartPos + ThisFreeSegment.Length)
+				{
+					StartFreeSegmentIndex = i;
+					bFoundSpaceForSegment = true;
+					break;
+				}
+				else if (!bFoundFutureSpaceForSegment &&
+					StartPos < ThisFreeSegment.StartPos &&
+					OtherUsedSegment.Length <= ThisFreeSegment.Length)
+				{
+					bFoundFutureSpaceForSegment = true;
+					FailedLength = ThisFreeSegment.StartPos - StartPos;
+					break;
+				}
 			}
-		}
 
-		if ( !bFoundSpaceForSegment )
-		{
-			if ( !bFoundFutureSpaceForSegment )
+			if (!bFoundSpaceForSegment)
 			{
-				FailedLength = Width;
-			}
+				if (!bFoundFutureSpaceForSegment)
+				{
+					FailedLength = Width;
+				}
 
-			return false;
+				return false;
+			}
 		}
-	}
-
-	return true;
-}
 
-void FAllocator2D::FlipX( FRect Rect )
-{
-	uint32 MinY = 0;
-
-	uint32 MaxY = Rect.H - 1;
-	for ( int32 Y = Rect.H - 1; Y > 0; --Y )
-	{
-		if ( Rows[ Y ].UsedSegments.Num() > 0 )
-			break;
-
-		--MaxY;
+		return true;
 	}
 
-	for ( uint32 Y = 0; Y <= MaxY; ++Y )
+	void FAllocator2D::FlipX(FRect Rect)
 	{
-		for ( uint32 LowX = 0; LowX < Rect.W / 2; ++LowX )
+		uint32 MinY = 0;
+
+		uint32 MaxY = Rect.H - 1;
+		for (int32 Y = Rect.H - 1; Y > 0; --Y)
 		{
-			uint32 HighX = Rect.W - 1 - LowX;
-		
-			const uint64 BitLow = GetBit( LowX, Y );
-			const uint64 BitHigh = GetBit( HighX, Y );
+			if (Rows[Y].UsedSegments.Num() > 0)
+				break;
 
-			if ( BitLow != 0ull )
-			{
-				SetBit( HighX, Y );
-			}
-			else
-			{
-				ClearBit( HighX, Y );
-			}
+			--MaxY;
+		}
 
-			if ( BitHigh != 0ull )
-			{
-				SetBit( LowX, Y );
-			}
-			else
+		for (uint32 Y = 0; Y <= MaxY; ++Y)
+		{
+			for (uint32 LowX = 0; LowX < Rect.W / 2; ++LowX)
 			{
-				ClearBit( LowX, Y );
-			}
-		}
-	}
+				uint32 HighX = Rect.W - 1 - LowX;
 
-	CreateUsedSegments();
-}
+				const uint64 BitLow = GetBit(LowX, Y);
+				const uint64 BitHigh = GetBit(HighX, Y);
 
-void FAllocator2D::FlipY( FRect Rect )
-{
-	uint32 MinY = 0;
+				if (BitLow != 0ull)
+				{
+					SetBit(HighX, Y);
+				}
+				else
+				{
+					ClearBit(HighX, Y);
+				}
 
-	uint32 MaxY = Rect.H - 1;
-	for ( int32 Y = Rect.H - 1; Y > 0; --Y )
-	{
-		if ( Rows[ Y ].UsedSegments.Num() > 0 )
-			break;
+				if (BitHigh != 0ull)
+				{
+					SetBit(LowX, Y);
+				}
+				else
+				{
+					ClearBit(LowX, Y);
+				}
+			}
+		}
 
-		--MaxY;
+		CreateUsedSegments();
 	}
 
-	for ( uint32 LowY = 0; LowY < ( MaxY + 1 ) / 2; ++LowY )
+	void FAllocator2D::FlipY(FRect Rect)
 	{
-		for ( uint32 X = 0; X < Rect.W; ++X )
+		uint32 MinY = 0;
+
+		uint32 MaxY = Rect.H - 1;
+		for (int32 Y = Rect.H - 1; Y > 0; --Y)
 		{
-			uint32 HighY = MaxY - LowY;
+			if (Rows[Y].UsedSegments.Num() > 0)
+				break;
 
-			const uint64 BitLow = GetBit( X, LowY );
-			const uint64 BitHigh = GetBit( X, HighY );
+			--MaxY;
+		}
 
-			if ( BitLow != 0ull )
-			{
-				SetBit( X, HighY );
-			}
-			else
+		for (uint32 LowY = 0; LowY < (MaxY + 1) / 2; ++LowY)
+		{
+			for (uint32 X = 0; X < Rect.W; ++X)
 			{
-				ClearBit( X, HighY );
-			}
+				uint32 HighY = MaxY - LowY;
 
-			if ( BitHigh != 0ull )
-			{
-				SetBit( X, LowY );
-			}
-			else
-			{
-				ClearBit( X, LowY );
+				const uint64 BitLow = GetBit(X, LowY);
+				const uint64 BitHigh = GetBit(X, HighY);
+
+				if (BitLow != 0ull)
+				{
+					SetBit(X, HighY);
+				}
+				else
+				{
+					ClearBit(X, HighY);
+				}
+
+				if (BitHigh != 0ull)
+				{
+					SetBit(X, LowY);
+				}
+				else
+				{
+					ClearBit(X, LowY);
+				}
 			}
 		}
-	}
 
-	for ( uint32 LowY = MinY; LowY < ( MaxY + 1 ) / 2; ++LowY )
-	{
-		const int32 HighY = MaxY - LowY;
+		for (uint32 LowY = MinY; LowY < (MaxY + 1) / 2; ++LowY)
+		{
+			const int32 HighY = MaxY - LowY;
 
-		Rows.Swap( LowY, HighY );
-		Rows[ LowY ].Index = LowY;
-		Rows[ HighY ].Index = HighY;
+			Rows.Swap(LowY, HighY);
+			Rows[LowY].Index = LowY;
+			Rows[HighY].Index = HighY;
+		}
 	}
-}
 
-void FAllocator2D::InitSegments()
-{
-	FSegment FreeSegment;
-	FreeSegment.StartPos = 0;
-	FreeSegment.Length = Width;
-
-	for ( FRow& Row : Rows )
+	void FAllocator2D::InitSegments()
 	{
-		Row.FreeSegments.Empty( 1 );
-		Row.FreeSegments.Add( FreeSegment );
-		Row.LongestSegment = FreeSegment.Length;
+		FSegment FreeSegment;
+		FreeSegment.StartPos = 0;
+		FreeSegment.Length = Width;
+
+		for (FRow& Row : Rows)
+		{
+			Row.FreeSegments.Empty(1);
+			Row.FreeSegments.Add(FreeSegment);
+			Row.LongestSegment = FreeSegment.Length;
 
-		Row.UsedSegments.Empty();
+			Row.UsedSegments.Empty();
 
+		}
 	}
-}
 
-void FAllocator2D::CreateUsedSegments()
-{
-	for ( uint32 y = 0; y < Height; ++y )
+	void FAllocator2D::CreateUsedSegments()
 	{
-		FRow& CurrentRow = Rows[y];
-		CurrentRow.LongestSegment = 0;
-		CurrentRow.UsedSegments.Empty();
-
-		int32 FirstUsedX = -1;
-		for ( uint32 k = 0; k < Pitch; ++k )
+		for (uint32 y = 0; y < Height; ++y)
 		{
-			uint32 x = k * 64;
+			FRow& CurrentRow = Rows[y];
+			CurrentRow.LongestSegment = 0;
+			CurrentRow.UsedSegments.Empty();
 
-			// If all bits are set
-			if ( Bits[k + y * Pitch] == ~0ull )
+			int32 FirstUsedX = -1;
+			for (uint32 k = 0; k < Pitch; ++k)
 			{
-				if ( FirstUsedX < 0 )
+				uint32 x = k * 64;
+
+				// If all bits are set
+				if (Bits[k + y * Pitch] == ~0ull)
 				{
-					FirstUsedX = x;
+					if (FirstUsedX < 0)
+					{
+						FirstUsedX = x;
+					}
+
+					if (k == Pitch - 1)
+					{
+						AddUsedSegment(CurrentRow, FirstUsedX, x + 64 - FirstUsedX);
+						FirstUsedX = -1;
+					}
 				}
-					
-				if ( k == Pitch - 1 )
+				// No bits are set
+				else if (Bits[k + y * Pitch] == 0ull)
 				{
-					AddUsedSegment( CurrentRow, FirstUsedX, x + 64 - FirstUsedX );
-					FirstUsedX = -1;
+					if (FirstUsedX >= 0)
+					{
+						AddUsedSegment(CurrentRow, FirstUsedX, x - FirstUsedX);
+						FirstUsedX = -1;
+					}
 				}
-			}
-			// No bits are set
-			else if ( Bits[k + y * Pitch] == 0ull )
-			{
-				if ( FirstUsedX >= 0 )
+				// Some bits are set
+				else
 				{
-					AddUsedSegment( CurrentRow, FirstUsedX, x - FirstUsedX );
-					FirstUsedX = -1;
-				}
-			}
-			// Some bits are set
-			else
-			{
-				for ( uint32 i = 0; i < 64; ++i )
-				{
-					const uint32 SubX = x + i;
-
-					if ( GetBit(SubX, y) )
+					for (uint32 i = 0; i < 64; ++i)
 					{
-						if ( FirstUsedX < 0 )
+						const uint32 SubX = x + i;
+
+						if (GetBit(SubX, y))
 						{
-							FirstUsedX = SubX;
+							if (FirstUsedX < 0)
+							{
+								FirstUsedX = SubX;
+							}
+
+							if (SubX == Width - 1)
+							{
+								AddUsedSegment(CurrentRow, FirstUsedX, SubX + 1 - FirstUsedX);
+								FirstUsedX = -1;
+							}
 						}
-							
-						if ( SubX == Width - 1 )
+						else if (FirstUsedX >= 0)
 						{
-							AddUsedSegment( CurrentRow, FirstUsedX, SubX + 1 - FirstUsedX );
+							AddUsedSegment(CurrentRow, FirstUsedX, SubX - FirstUsedX);
 							FirstUsedX = -1;
 						}
 					}
-					else if ( FirstUsedX >= 0 )
-					{
-						AddUsedSegment( CurrentRow, FirstUsedX, SubX - FirstUsedX );
-						FirstUsedX = -1;
-					}
 				}
 			}
 		}
 	}
-}
 
-void FAllocator2D::AddUsedSegment( FRow& Row, uint32 StartPos, uint32 Length )
-{
-	FSegment UsedSegment;
-	UsedSegment.StartPos = StartPos;
-	UsedSegment.Length = Length;
+	void FAllocator2D::AddUsedSegment(FRow& Row, uint32 StartPos, uint32 Length)
+	{
+		FSegment UsedSegment;
+		UsedSegment.StartPos = StartPos;
+		UsedSegment.Length = Length;
 
-	Row.UsedSegments.Add(UsedSegment);
+		Row.UsedSegments.Add(UsedSegment);
 
-	if ( UsedSegment.Length > Row.LongestSegment )
-	{
-		Row.LongestSegment = UsedSegment.Length;
+		if (UsedSegment.Length > Row.LongestSegment)
+		{
+			Row.LongestSegment = UsedSegment.Length;
+		}
 	}
-}
 
-void FAllocator2D::MergeSegments( FRect Rect, const FAllocator2D& Other )
-{
-	for ( uint32 y = 0; y < Rect.H; ++y )
+	void FAllocator2D::MergeSegments(FRect Rect, const FAllocator2D& Other)
 	{
-		FRow& ThisRow = Rows[Rect.Y + y];
-		const FRow& OtherRow = Other.Rows[y];
-
-		for ( const FSegment& OtherUsedSegment : OtherRow.UsedSegments )
+		for (uint32 y = 0; y < Rect.H; ++y)
 		{
-			for ( int32 i = 0; i < ThisRow.FreeSegments.Num(); ++i )
+			FRow& ThisRow = Rows[Rect.Y + y];
+			const FRow& OtherRow = Other.Rows[y];
+
+			for (const FSegment& OtherUsedSegment : OtherRow.UsedSegments)
 			{
-				FSegment& ThisFreeSegment = ThisRow.FreeSegments[i];
+				for (int32 i = 0; i < ThisRow.FreeSegments.Num(); ++i)
+				{
+					FSegment& ThisFreeSegment = ThisRow.FreeSegments[i];
 
-				uint32 StartPos = Rect.X + OtherUsedSegment.StartPos;
+					uint32 StartPos = Rect.X + OtherUsedSegment.StartPos;
 
-				if ( StartPos >= ThisFreeSegment.StartPos &&
-					StartPos < ThisFreeSegment.StartPos + ThisFreeSegment.Length )
-				{
-					if ( ThisFreeSegment.Length == 1 )
+					if (StartPos >= ThisFreeSegment.StartPos &&
+						StartPos < ThisFreeSegment.StartPos + ThisFreeSegment.Length)
 					{
-						ThisRow.FreeSegments.RemoveAtSwap(i);
-						break;
-					}
+						if (ThisFreeSegment.Length == 1)
+						{
+							ThisRow.FreeSegments.RemoveAtSwap(i);
+							break;
+						}
 
-					FSegment FirstSegment;
-					FirstSegment.StartPos = ThisFreeSegment.StartPos;
-					FirstSegment.Length = StartPos - ThisFreeSegment.StartPos;
+						FSegment FirstSegment;
+						FirstSegment.StartPos = ThisFreeSegment.StartPos;
+						FirstSegment.Length = StartPos - ThisFreeSegment.StartPos;
 
-					uint32 EndPos = Rect.X + FMath::Min( OtherUsedSegment.StartPos + OtherUsedSegment.Length, Rect.W ) - 1;
+						uint32 EndPos = Rect.X + FMath::Min(OtherUsedSegment.StartPos + OtherUsedSegment.Length, Rect.W) - 1;
 
-					FSegment SecondSegment;
-					SecondSegment.StartPos = EndPos + 1;
-					SecondSegment.Length = ThisFreeSegment.StartPos + ThisFreeSegment.Length - SecondSegment.StartPos;
+						FSegment SecondSegment;
+						SecondSegment.StartPos = EndPos + 1;
+						SecondSegment.Length = ThisFreeSegment.StartPos + ThisFreeSegment.Length - SecondSegment.StartPos;
 
-					ThisRow.FreeSegments.RemoveAtSwap(i);
+						ThisRow.FreeSegments.RemoveAtSwap(i);
 
-					if ( FirstSegment.Length > 0 )
-					{
-						ThisRow.FreeSegments.Add( FirstSegment );
-					}
-					
-					if ( SecondSegment.Length > 0 )
-					{
-						ThisRow.FreeSegments.Add( SecondSegment );
-					}
+						if (FirstSegment.Length > 0)
+						{
+							ThisRow.FreeSegments.Add(FirstSegment);
+						}
 
-					break;
+						if (SecondSegment.Length > 0)
+						{
+							ThisRow.FreeSegments.Add(SecondSegment);
+						}
+
+						break;
+					}
 				}
 			}
-		}
 
-		ThisRow.FreeSegments.Sort();
+			ThisRow.FreeSegments.Sort();
 
-		ThisRow.LongestSegment = 0;
-		for ( const FSegment& Segment : ThisRow.FreeSegments )
-		{
-			ThisRow.LongestSegment = FMath::Max( Segment.Length, ThisRow.LongestSegment );
+			ThisRow.LongestSegment = 0;
+			for (const FSegment& Segment : ThisRow.FreeSegments)
+			{
+				ThisRow.LongestSegment = FMath::Max(Segment.Length, ThisRow.LongestSegment);
+			}
 		}
 	}
-}
+}
\ No newline at end of file
diff --git a/Engine/Source/Developer/MeshDescriptionOperations/Private/Allocator2D.h b/Engine/Source/Developer/MeshDescriptionOperations/Private/Allocator2D.h
--- a/Engine/Source/Developer/MeshDescriptionOperations/Private/Allocator2D.h
+++ b/Engine/Source/Developer/MeshDescriptionOperations/Private/Allocator2D.h
@@ -7,152 +7,155 @@
 struct FRect;
 struct Rect;
 
-class FAllocator2D
+namespace MeshDescriptionOp
 {
-public:
-	struct FRect
+	class FAllocator2D
 	{
-		uint32 X;
-		uint32 Y;
-		uint32 W;
-		uint32 H;
-	};
+	public:
+		struct FRect
+		{
+			uint32 X;
+			uint32 Y;
+			uint32 W;
+			uint32 H;
+		};
 
-	struct FSegment
-	{
-		uint32 StartPos;
-		uint32 Length;
+		struct FSegment
+		{
+			uint32 StartPos;
+			uint32 Length;
 
-		bool operator<( const FSegment& Other ) const { return StartPos < Other.StartPos; }
-	};
+			bool operator<(const FSegment& Other) const { return StartPos < Other.StartPos; }
+		};
 
-	struct FRow
-	{
-		uint32 Index;
-		uint32 LongestSegment; // Represents either the longest free segment or the longest used segment, depending on how we're using this row
+		struct FRow
+		{
+			uint32 Index;
+			uint32 LongestSegment; // Represents either the longest free segment or the longest used segment, depending on how we're using this row
+
+			TArray< FSegment > FreeSegments;
+			TArray< FSegment > UsedSegments;
+		};
+
+	public:
+		FAllocator2D(uint32 Width, uint32 Height);
+		FAllocator2D(const FAllocator2D& Other);
+		~FAllocator2D();
+
+		FAllocator2D& operator=(const FAllocator2D& Other);
+
+		// Must clear before using
+		void		Clear();
+
+		bool		Find(FRect& Rect);
+		bool		Test(FRect Rect);
+		void		Alloc(FRect Rect);
 
-		TArray< FSegment > FreeSegments;
-		TArray< FSegment > UsedSegments;
+		bool		FindBitByBit(FRect& Rect, const FAllocator2D& Other);
+		bool		FindWithSegments(FRect& Rect, FRect BestRect, const FAllocator2D& Other);
+		bool		Test(FRect Rect, const FAllocator2D& Other);
+		void		Alloc(FRect Rect, const FAllocator2D& Other);
+
+		uint64		GetBit(uint32 x, uint32 y) const;
+		void		SetBit(uint32 x, uint32 y);
+		void		ClearBit(uint32 x, uint32 y);
+
+		void		CreateUsedSegments();
+		void		MergeSegments(FRect Rect, const FAllocator2D& Other);
+
+		void		FlipX(FRect Rect);
+		void		FlipY(FRect Rect);
+
+	protected:
+		bool		TestAllRows(FRect Rect, const FAllocator2D& Other, uint32& FailedLength);
+		bool		TestRow(const FRow& ThisRow, const FRow& OtherRow, FRect Rect, uint32& FailedLength);
+
+		void		InitSegments();
+		void		AddUsedSegment(FRow& Row, uint32 StartPos, uint32 Length);
+
+	private:
+		uint64*		Bits;
+
+		uint32		Width;
+		uint32		Height;
+		uint32		Pitch;
+
+		TArray< FRow > Rows;
+		int32		LastRowFail;
 	};
 
-public:
-				FAllocator2D( uint32 Width, uint32 Height );
-				FAllocator2D( const FAllocator2D& Other );
-				~FAllocator2D();
-	
-	FAllocator2D& operator=( const FAllocator2D& Other );
-
-	// Must clear before using
-	void		Clear();
-
-	bool		Find( FRect& Rect );
-	bool		Test( FRect Rect );
-	void		Alloc( FRect Rect );
-
-	bool		FindBitByBit( FRect& Rect, const FAllocator2D& Other );
-	bool		FindWithSegments( FRect& Rect, FRect BestRect, const FAllocator2D& Other );
-	bool		Test( FRect Rect, const FAllocator2D& Other );
-	void		Alloc( FRect Rect, const FAllocator2D& Other );
-	
-	uint64		GetBit( uint32 x, uint32 y ) const;
-	void		SetBit( uint32 x, uint32 y );
-	void		ClearBit( uint32 x, uint32 y );
-	
-	void		CreateUsedSegments();
-	void		MergeSegments( FRect Rect, const FAllocator2D& Other );
-
-	void		FlipX( FRect Rect );
-	void		FlipY( FRect Rect );
-
-protected:
-	bool		TestAllRows( FRect Rect, const FAllocator2D& Other, uint32& FailedLength );
-	bool		TestRow( const FRow& ThisRow, const FRow& OtherRow, FRect Rect, uint32& FailedLength );
-
-	void		InitSegments();
-	void		AddUsedSegment( FRow& Row, uint32 StartPos, uint32 Length );
-
-private:
-	uint64*		Bits;
-
-	uint32		Width;
-	uint32		Height;
-	uint32		Pitch;
-
-	TArray< FRow > Rows;
-	int32		LastRowFail;
-};
-
-// Returns non-zero if set
-FORCEINLINE uint64 FAllocator2D::GetBit( uint32 x, uint32 y ) const
-{
-	return Bits[ (x >> 6) + y * Pitch ] & ( 1ull << ( x & 63 ) );
-}
+	// Returns non-zero if set
+	FORCEINLINE uint64 FAllocator2D::GetBit(uint32 x, uint32 y) const
+	{
+		return Bits[(x >> 6) + y * Pitch] & (1ull << (x & 63));
+	}
 
-FORCEINLINE void FAllocator2D::SetBit( uint32 x, uint32 y )
-{
-	Bits[ (x >> 6) + y * Pitch ] |= ( 1ull << ( x & 63 ) );
-}
+	FORCEINLINE void FAllocator2D::SetBit(uint32 x, uint32 y)
+	{
+		Bits[(x >> 6) + y * Pitch] |= (1ull << (x & 63));
+	}
 
-FORCEINLINE void FAllocator2D::ClearBit( uint32 x, uint32 y )
-{
-	Bits[ (x >> 6) + y * Pitch ] &= ~( 1ull << ( x & 63 ) );
-}
+	FORCEINLINE void FAllocator2D::ClearBit(uint32 x, uint32 y)
+	{
+		Bits[(x >> 6) + y * Pitch] &= ~(1ull << (x & 63));
+	}
 
-inline bool FAllocator2D::Test( FRect Rect )
-{
-	for( uint32 y = Rect.Y; y < Rect.Y + Rect.H; y++ )
+	inline bool FAllocator2D::Test(FRect Rect)
 	{
-		for( uint32 x = Rect.X; x < Rect.X + Rect.W; x++ )
+		for (uint32 y = Rect.Y; y < Rect.Y + Rect.H; y++)
 		{
-			if( GetBit( x, y ) )
+			for (uint32 x = Rect.X; x < Rect.X + Rect.W; x++)
 			{
-				return false;
+				if (GetBit(x, y))
+				{
+					return false;
+				}
 			}
 		}
-	}
-	
-	return true;
-}
 
-inline bool FAllocator2D::Test( FRect Rect, const FAllocator2D& Other )
-{
-	const uint32 LowShift = Rect.X & 63;
-	const uint32 HighShift = 64 - LowShift;
+		return true;
+	}
 
-	for( uint32 y = 0; y < Rect.H; y++ )
+	inline bool FAllocator2D::Test(FRect Rect, const FAllocator2D& Other)
 	{
-#if 1
-		uint32 ThisIndex = (Rect.X >> 6) + (y + Rect.Y) * Pitch;
-		uint32 OtherIndex = y * Pitch;
+		const uint32 LowShift = Rect.X & 63;
+		const uint32 HighShift = 64 - LowShift;
 
-		// Test a uint64 at a time
-		for( uint32 x = 0; x < Rect.W; x += 64 )
+		for (uint32 y = 0; y < Rect.H; y++)
 		{
-			// no need to zero out HighInt on wrap around because Other will always be zero outside Rect.
-			uint64 LowInt  = Bits[ ThisIndex ];
-			uint64 HighInt = Bits[ ThisIndex + 1 ];
-
-			uint64 ThisInt = (HighInt << HighShift) | (LowInt >> LowShift);
-			uint64 OtherInt = Other.Bits[ OtherIndex ];
+#if 1
+			uint32 ThisIndex = (Rect.X >> 6) + (y + Rect.Y) * Pitch;
+			uint32 OtherIndex = y * Pitch;
 
-			if( ThisInt & OtherInt )
+			// Test a uint64 at a time
+			for (uint32 x = 0; x < Rect.W; x += 64)
 			{
-				return false;
-			}
+				// no need to zero out HighInt on wrap around because Other will always be zero outside Rect.
+				uint64 LowInt = Bits[ThisIndex];
+				uint64 HighInt = Bits[ThisIndex + 1];
 
-			ThisIndex++;
-			OtherIndex++;
-		}
+				uint64 ThisInt = (HighInt << HighShift) | (LowInt >> LowShift);
+				uint64 OtherInt = Other.Bits[OtherIndex];
+
+				if (ThisInt & OtherInt)
+				{
+					return false;
+				}
+
+				ThisIndex++;
+				OtherIndex++;
+			}
 #else
-		for( uint32 x = 0; x < Rect.W; x++ )
-		{
-			if( Other.GetBit( x, y ) && GetBit( x + Rect.X, y + Rect.Y ) )
+			for (uint32 x = 0; x < Rect.W; x++)
 			{
-				return false;
+				if (Other.GetBit(x, y) && GetBit(x + Rect.X, y + Rect.Y))
+				{
+					return false;
+				}
 			}
-		}
 #endif
+		}
+
+		return true;
 	}
-	
-	return true;
-}
+}
\ No newline at end of file
diff --git a/Engine/Source/Developer/MeshDescriptionOperations/Private/DisjointSet.h b/Engine/Source/Developer/MeshDescriptionOperations/Private/DisjointSet.h
--- a/Engine/Source/Developer/MeshDescriptionOperations/Private/DisjointSet.h
+++ b/Engine/Source/Developer/MeshDescriptionOperations/Private/DisjointSet.h
@@ -4,79 +4,82 @@
 
 #include "CoreMinimal.h"
 
-class FDisjointSet
+namespace MeshDescriptionOp
 {
-public:
-			FDisjointSet( const uint32 Size );
-	
-	void	Union( uint32 x, uint32 y );
-	uint32	Find( uint32 i );
+	class FDisjointSet
+	{
+	public:
+		FDisjointSet(const uint32 Size);
 
-	uint32	operator[]( uint32 i ) const	{ return Parents[i]; }
+		void	Union(uint32 x, uint32 y);
+		uint32	Find(uint32 i);
 
-private:
-	TArray<	uint32 >	Parents;
-};
+		uint32	operator[](uint32 i) const { return Parents[i]; }
 
-FDisjointSet::FDisjointSet( const uint32 Size )
-{
-	Parents.SetNumUninitialized( Size );
-	for( uint32 i = 0; i < Size; i++ )
+	private:
+		TArray<	uint32 >	Parents;
+	};
+
+	FDisjointSet::FDisjointSet(const uint32 Size)
 	{
-		Parents[i] = i;
+		Parents.SetNumUninitialized(Size);
+		for (uint32 i = 0; i < Size; i++)
+		{
+			Parents[i] = i;
+		}
 	}
-}
 
-// Union with splicing
-void FDisjointSet::Union( uint32 x, uint32 y )
-{
-	while( Parents[x] != Parents[y] )
+	// Union with splicing
+	void FDisjointSet::Union(uint32 x, uint32 y)
 	{
-		// Pick larger
-		if( Parents[x] < Parents[y] )
+		while (Parents[x] != Parents[y])
 		{
-			uint32 p = Parents[x];
-			Parents[x] = Parents[y];
-			if( x == p )
+			// Pick larger
+			if (Parents[x] < Parents[y])
 			{
-				return;
+				uint32 p = Parents[x];
+				Parents[x] = Parents[y];
+				if (x == p)
+				{
+					return;
+				}
+				x = p;
 			}
-			x = p;
-		}
-		else
-		{
-			uint32 p = Parents[y];
-			Parents[y] = Parents[x];
-			if( y == p )
+			else
 			{
-				return;
+				uint32 p = Parents[y];
+				Parents[y] = Parents[x];
+				if (y == p)
+				{
+					return;
+				}
+				y = p;
 			}
-			y = p;
 		}
 	}
-}
 
-// Find with path compression
-uint32 FDisjointSet::Find( uint32 i )
-{
-	// Find root
-	uint32 Start = i;
-	uint32 Root = Parents[i];
-	while( Root != i )
+	// Find with path compression
+	uint32 FDisjointSet::Find(uint32 i)
 	{
-		i = Root;
-		Root = Parents[i];
-	}
+		// Find root
+		uint32 Start = i;
+		uint32 Root = Parents[i];
+		while (Root != i)
+		{
+			i = Root;
+			Root = Parents[i];
+		}
 
-	// Point all nodes on path to root
-	i = Start;
-	uint32 Parent = Parents[i];
-	while( Parent != Root )
-	{
-		Parents[i] = Root;
-		i = Parent;
-		Parent = Parents[i];
-	}
+		// Point all nodes on path to root
+		i = Start;
+		uint32 Parent = Parents[i];
+		while (Parent != Root)
+		{
+			Parents[i] = Root;
+			i = Parent;
+			Parent = Parents[i];
+		}
 
-	return Root;
-}
+		return Root;
+	}
+}
\ No newline at end of file
diff --git a/Engine/Source/Developer/MeshDescriptionOperations/Private/LayoutUV.cpp b/Engine/Source/Developer/MeshDescriptionOperations/Private/LayoutUV.cpp
--- a/Engine/Source/Developer/MeshDescriptionOperations/Private/LayoutUV.cpp
+++ b/Engine/Source/Developer/MeshDescriptionOperations/Private/LayoutUV.cpp
@@ -9,1078 +9,1080 @@
 DEFINE_LOG_CATEGORY_STATIC(LogMeshDescriptionLayoutUV, Warning, All);
 
 #define CHART_JOINING	1
-
-FLayoutUV::FLayoutUV( UMeshDescription* InMesh, uint32 InSrcChannel, uint32 InDstChannel, uint32 InTextureResolution )
-	: MeshDescription( InMesh )
-	, SrcChannel( InSrcChannel )
-	, DstChannel( InDstChannel )
-	, TextureResolution( InTextureResolution )
-	, TotalUVArea( 0.0f )
-	, LayoutRaster( TextureResolution, TextureResolution )
-	, ChartRaster( TextureResolution, TextureResolution )
-	, BestChartRaster( TextureResolution, TextureResolution )
-	, ChartShader( &ChartRaster )
-	, LayoutVersion(FMeshDescriptionOperations::ELightmapUVVersion::Latest )
-{}
-
-void FLayoutUV::FindCharts( const TMultiMap<int32,int32>& OverlappingCorners )
+namespace MeshDescriptionOp
 {
-	const float ThreshUVsAreSame = GetUVEqualityThreshold();
-	double Begin = FPlatformTime::Seconds();
-
-	uint32 NumIndexes = MeshDescription->VertexInstances().Num();
-	uint32 NumTris = NumIndexes / 3;
-	check(MeshDescription->Polygons().Num() == NumTris);
+	FLayoutUV::FLayoutUV(UMeshDescription* InMesh, uint32 InSrcChannel, uint32 InDstChannel, uint32 InTextureResolution)
+		: MeshDescription(InMesh)
+		, SrcChannel(InSrcChannel)
+		, DstChannel(InDstChannel)
+		, TextureResolution(InTextureResolution)
+		, TotalUVArea(0.0f)
+		, LayoutRaster(TextureResolution, TextureResolution)
+		, ChartRaster(TextureResolution, TextureResolution)
+		, BestChartRaster(TextureResolution, TextureResolution)
+		, ChartShader(&ChartRaster)
+		, LayoutVersion(FMeshDescriptionOperations::ELightmapUVVersion::Latest)
+	{}
+
+	void FLayoutUV::FindCharts(const TMultiMap<int32, int32>& OverlappingCorners)
+	{
+		const float ThreshUVsAreSame = GetUVEqualityThreshold();
+		double Begin = FPlatformTime::Seconds();
 
-	TArray< int32 > TranslatedMatches;
-	TranslatedMatches.SetNumUninitialized( NumIndexes );
-	TexCoords.SetNumUninitialized( NumIndexes );
-	int32 VertexInstanceIndex = 0;
+		uint32 NumIndexes = MeshDescription->VertexInstances().Num();
+		uint32 NumTris = NumIndexes / 3;
+		check(MeshDescription->Polygons().Num() == NumTris);
 
-	const TVertexInstanceAttributeArray<FVector2D>& VertexUVs = MeshDescription->VertexInstanceAttributes().GetAttributes<FVector2D>(MeshAttribute::VertexInstance::TextureCoordinate, SrcChannel);
-	for (const FVertexInstanceID VertexInstanceID : MeshDescription->VertexInstances().GetElementIDs())
-	{
-		check(VertexInstanceIndex == VertexInstanceID.GetValue());
-		TranslatedMatches[VertexInstanceIndex] = -1;
-		TexCoords[VertexInstanceIndex] = VertexUVs[VertexInstanceID];
-		VertexInstanceIndex++;
-	}
+		TArray< int32 > TranslatedMatches;
+		TranslatedMatches.SetNumUninitialized(NumIndexes);
+		TexCoords.SetNumUninitialized(NumIndexes);
+		int32 VertexInstanceIndex = 0;
 
-	// Build disjoint set
-	FDisjointSet DisjointSet( NumTris );
-	for( uint32 i = 0; i < NumIndexes; i++ )
-	{
-		for( auto It = OverlappingCorners.CreateConstKeyIterator(i); It; ++It )
+		const TVertexInstanceAttributeArray<FVector2D>& VertexUVs = MeshDescription->VertexInstanceAttributes().GetAttributes<FVector2D>(MeshAttribute::VertexInstance::TextureCoordinate, SrcChannel);
+		for (const FVertexInstanceID VertexInstanceID : MeshDescription->VertexInstances().GetElementIDs())
 		{
-			uint32 j = It.Value();
+			check(VertexInstanceIndex == VertexInstanceID.GetValue());
+			TranslatedMatches[VertexInstanceIndex] = -1;
+			TexCoords[VertexInstanceIndex] = VertexUVs[VertexInstanceID];
+			VertexInstanceIndex++;
+		}
 
-			if( j > i )
+		// Build disjoint set
+		FDisjointSet DisjointSet(NumTris);
+		for (uint32 i = 0; i < NumIndexes; i++)
+		{
+			for (auto It = OverlappingCorners.CreateConstKeyIterator(i); It; ++It)
 			{
-				const uint32 TriI = i/3;
-				const uint32 TriJ = j/3;
+				uint32 j = It.Value();
+
+				if (j > i)
+				{
+					const uint32 TriI = i / 3;
+					const uint32 TriJ = j / 3;
 
-				bool bUnion = false;
+					bool bUnion = false;
 
 #if CHART_JOINING
-				bool bPositionMatch = PositionsMatch( i, j );
-				if( bPositionMatch )
-				{
-					uint32 i1 = 3 * TriI + (i + 1) % 3;
-					uint32 i2 = 3 * TriI + (i + 2) % 3;
-					uint32 j1 = 3 * TriJ + (j + 1) % 3;
-					uint32 j2 = 3 * TriJ + (j + 2) % 3;
-
-					bool bEdgeMatch21 = PositionsMatch( i2, j1 );
-					bool bEdgeMatch12 = PositionsMatch( i1, j2 );
-					if( bEdgeMatch21 || bEdgeMatch12 )
+					bool bPositionMatch = PositionsMatch(i, j);
+					if (bPositionMatch)
 					{
-						uint32 ie = bEdgeMatch21 ? i2 : i1;
-						uint32 je = bEdgeMatch21 ? j1 : j2;
-							
-						bool bUVMatch = UVsMatch( i, j ) && UVsMatch( ie, je );
-						bool bUVWindingMatch = TriangleUVArea( TriI ) * TriangleUVArea( TriJ ) >= 0.0f;
-						if( bUVMatch && bUVWindingMatch )
+						uint32 i1 = 3 * TriI + (i + 1) % 3;
+						uint32 i2 = 3 * TriI + (i + 2) % 3;
+						uint32 j1 = 3 * TriJ + (j + 1) % 3;
+						uint32 j2 = 3 * TriJ + (j + 2) % 3;
+
+						bool bEdgeMatch21 = PositionsMatch(i2, j1);
+						bool bEdgeMatch12 = PositionsMatch(i1, j2);
+						if (bEdgeMatch21 || bEdgeMatch12)
 						{
-							bUnion = true;
-						}
-						else if( NormalsMatch( i, j ) && NormalsMatch( ie, je ) )
-						{
-							// Chart edge
-							FVector2D EdgeUVi = TexCoords[ie] - TexCoords[i];
-							FVector2D EdgeUVj = TexCoords[je] - TexCoords[j];
-							
-							// Would these edges match if the charts were translated
-							bool bTranslatedUVMatch = ( EdgeUVi - EdgeUVj ).IsNearlyZero( ThreshUVsAreSame );
-							if( bTranslatedUVMatch )
+							uint32 ie = bEdgeMatch21 ? i2 : i1;
+							uint32 je = bEdgeMatch21 ? j1 : j2;
+
+							bool bUVMatch = UVsMatch(i, j) && UVsMatch(ie, je);
+							bool bUVWindingMatch = TriangleUVArea(TriI) * TriangleUVArea(TriJ) >= 0.0f;
+							if (bUVMatch && bUVWindingMatch)
+							{
+								bUnion = true;
+							}
+							else if (NormalsMatch(i, j) && NormalsMatch(ie, je))
 							{
-								// Note: may be mirrored
-								
-								// TODO should these be restricted to axis aligned edges?
-								uint32 EdgeI = bEdgeMatch21 ? i2 : i;
-								uint32 EdgeJ = bEdgeMatch21 ? j : j2;
-
-								// Only allow one match per edge
-								if( TranslatedMatches[ EdgeI ] < 0 &&
-									TranslatedMatches[ EdgeJ ] < 0 )
+								// Chart edge
+								FVector2D EdgeUVi = TexCoords[ie] - TexCoords[i];
+								FVector2D EdgeUVj = TexCoords[je] - TexCoords[j];
+
+								// Would these edges match if the charts were translated
+								bool bTranslatedUVMatch = (EdgeUVi - EdgeUVj).IsNearlyZero(ThreshUVsAreSame);
+								if (bTranslatedUVMatch)
 								{
-									TranslatedMatches[ EdgeI ] = EdgeJ;
-									TranslatedMatches[ EdgeJ ] = EdgeI;
+									// Note: may be mirrored
+
+									// TODO should these be restricted to axis aligned edges?
+									uint32 EdgeI = bEdgeMatch21 ? i2 : i;
+									uint32 EdgeJ = bEdgeMatch21 ? j : j2;
+
+									// Only allow one match per edge
+									if (TranslatedMatches[EdgeI] < 0 &&
+										TranslatedMatches[EdgeJ] < 0)
+									{
+										TranslatedMatches[EdgeI] = EdgeJ;
+										TranslatedMatches[EdgeJ] = EdgeI;
+									}
 								}
 							}
 						}
 					}
-				}
 #else
-				if( VertsMatch( i, j ) )
-				{
-					// Edge must match as well (same winding)
-					if( VertsMatch( 3 * TriI + (i - 1) % 3, 3 * TriJ + (j + 1) % 3 ) ||
-						VertsMatch( 3 * TriI + (i + 1) % 3, 3 * TriJ + (j - 1) % 3 ) )
+					if (VertsMatch(i, j))
 					{
-						// Check for UV winding match too
-						if( TriangleUVArea( TriI ) * TriangleUVArea( TriJ ) >= 0.0f )
+						// Edge must match as well (same winding)
+						if (VertsMatch(3 * TriI + (i - 1) % 3, 3 * TriJ + (j + 1) % 3) ||
+							VertsMatch(3 * TriI + (i + 1) % 3, 3 * TriJ + (j - 1) % 3))
 						{
-							bUnion = true;
+							// Check for UV winding match too
+							if (TriangleUVArea(TriI) * TriangleUVArea(TriJ) >= 0.0f)
+							{
+								bUnion = true;
+							}
 						}
 					}
-				}
 #endif
 
-				if( bUnion )
-				{
-					// TODO solve spiral case by checking sets for UV overlap
-					DisjointSet.Union( TriI, TriJ );
+					if (bUnion)
+					{
+						// TODO solve spiral case by checking sets for UV overlap
+						DisjointSet.Union(TriI, TriJ);
+					}
 				}
 			}
 		}
-	}
-
-	// Sort tris by chart
-	SortedTris.SetNumUninitialized( NumTris );
-	for( uint32 i = 0; i < NumTris; i++ )
-	{
-		// Flatten disjoint set path
-		DisjointSet.Find(i);
-		SortedTris[i] = i;
-	}
 
-	struct FCompareTris
-	{
-		FDisjointSet* DisjointSet;
-
-		FCompareTris( FDisjointSet* InDisjointSet )
-		: DisjointSet( InDisjointSet )
-		{}
-
-		FORCEINLINE bool operator()( uint32 A, uint32 B ) const
+		// Sort tris by chart
+		SortedTris.SetNumUninitialized(NumTris);
+		for (uint32 i = 0; i < NumTris; i++)
 		{
-			return (*DisjointSet)[A] < (*DisjointSet)[B];
+			// Flatten disjoint set path
+			DisjointSet.Find(i);
+			SortedTris[i] = i;
 		}
-	};
 
-	Algo::IntroSort( SortedTris, FCompareTris( &DisjointSet ) );
+		struct FCompareTris
+		{
+			FDisjointSet* DisjointSet;
 
-	TMap< uint32, int32 > DisjointSetToChartMap;
+			FCompareTris(FDisjointSet* InDisjointSet)
+				: DisjointSet(InDisjointSet)
+			{}
 
-	const TVertexAttributeArray<FVector>& VertexPositions = MeshDescription->VertexAttributes().GetAttributes<FVector>(MeshAttribute::Vertex::Position);
+			FORCEINLINE bool operator()(uint32 A, uint32 B) const
+			{
+				return (*DisjointSet)[A] < (*DisjointSet)[B];
+			}
+		};
 
-	// Build Charts
-	for( uint32 Tri = 0; Tri < NumTris; )
-	{
-		int32 i = Charts.AddUninitialized();
-		FMeshChart& Chart = Charts[i];
-		
-		Chart.MinUV = FVector2D( FLT_MAX, FLT_MAX );
-		Chart.MaxUV = FVector2D( -FLT_MAX, -FLT_MAX );
-		Chart.UVArea = 0.0f;
-		Chart.WorldScale = FVector2D::ZeroVector;
-		FMemory::Memset( Chart.Join, 0xff );
+		Algo::IntroSort(SortedTris, FCompareTris(&DisjointSet));
 
-		Chart.FirstTri = Tri;
+		TMap< uint32, int32 > DisjointSetToChartMap;
 
-		uint32 ChartID = DisjointSet[ SortedTris[ Tri ] ];
-		DisjointSetToChartMap.Add( ChartID, i );
+		const TVertexAttributeArray<FVector>& VertexPositions = MeshDescription->VertexAttributes().GetAttributes<FVector>(MeshAttribute::Vertex::Position);
 
-		for( ; Tri < NumTris && DisjointSet[ SortedTris[ Tri ] ] == ChartID; Tri++ )
+		// Build Charts
+		for (uint32 Tri = 0; Tri < NumTris; )
 		{
-			// Calculate chart bounds
-			FVector		Positions[3];
-			FVector2D	UVs[3];
-			for( int k = 0; k < 3; k++ )
+			int32 i = Charts.AddUninitialized();
+			FMeshChart& Chart = Charts[i];
+
+			Chart.MinUV = FVector2D(FLT_MAX, FLT_MAX);
+			Chart.MaxUV = FVector2D(-FLT_MAX, -FLT_MAX);
+			Chart.UVArea = 0.0f;
+			Chart.WorldScale = FVector2D::ZeroVector;
+			FMemory::Memset(Chart.Join, 0xff);
+
+			Chart.FirstTri = Tri;
+
+			uint32 ChartID = DisjointSet[SortedTris[Tri]];
+			DisjointSetToChartMap.Add(ChartID, i);
+
+			for (; Tri < NumTris && DisjointSet[SortedTris[Tri]] == ChartID; Tri++)
 			{
-				uint32 Index = 3 * SortedTris[ Tri ] + k;
+				// Calculate chart bounds
+				FVector		Positions[3];
+				FVector2D	UVs[3];
+				for (int k = 0; k < 3; k++)
+				{
+					uint32 Index = 3 * SortedTris[Tri] + k;
 
-				FVertexInstanceID VertexInstanceID(Index);
-				Positions[k] = VertexPositions[MeshDescription->GetVertexInstanceVertex(VertexInstanceID)];
-				UVs[k] = TexCoords[ Index ];
+					FVertexInstanceID VertexInstanceID(Index);
+					Positions[k] = VertexPositions[MeshDescription->GetVertexInstanceVertex(VertexInstanceID)];
+					UVs[k] = TexCoords[Index];
 
-				Chart.MinUV.X = FMath::Min( Chart.MinUV.X, UVs[k].X );
-				Chart.MinUV.Y = FMath::Min( Chart.MinUV.Y, UVs[k].Y );
-				Chart.MaxUV.X = FMath::Max( Chart.MaxUV.X, UVs[k].X );
-				Chart.MaxUV.Y = FMath::Max( Chart.MaxUV.Y, UVs[k].Y );
-			}
+					Chart.MinUV.X = FMath::Min(Chart.MinUV.X, UVs[k].X);
+					Chart.MinUV.Y = FMath::Min(Chart.MinUV.Y, UVs[k].Y);
+					Chart.MaxUV.X = FMath::Max(Chart.MaxUV.X, UVs[k].X);
+					Chart.MaxUV.Y = FMath::Max(Chart.MaxUV.Y, UVs[k].Y);
+				}
 
-			FVector Edge1 = Positions[1] - Positions[0];
-			FVector Edge2 = Positions[2] - Positions[0];
-			float Area = 0.5f * ( Edge1 ^ Edge2 ).Size();
+				FVector Edge1 = Positions[1] - Positions[0];
+				FVector Edge2 = Positions[2] - Positions[0];
+				float Area = 0.5f * (Edge1 ^ Edge2).Size();
 
-			FVector2D EdgeUV1 = UVs[1] - UVs[0];
-			FVector2D EdgeUV2 = UVs[2] - UVs[0];
-			float UVArea = 0.5f * FMath::Abs( EdgeUV1.X * EdgeUV2.Y - EdgeUV1.Y * EdgeUV2.X );
+				FVector2D EdgeUV1 = UVs[1] - UVs[0];
+				FVector2D EdgeUV2 = UVs[2] - UVs[0];
+				float UVArea = 0.5f * FMath::Abs(EdgeUV1.X * EdgeUV2.Y - EdgeUV1.Y * EdgeUV2.X);
 
-			FVector2D UVLength;
-			UVLength.X = ( EdgeUV2.Y * Edge1 - EdgeUV1.Y * Edge2 ).Size();
-			UVLength.Y = (-EdgeUV2.X * Edge1 + EdgeUV1.X * Edge2 ).Size();
-			
-			Chart.WorldScale += UVLength;
-			Chart.UVArea += UVArea;
-		}
-		
-		Chart.LastTri = Tri;
+				FVector2D UVLength;
+				UVLength.X = (EdgeUV2.Y * Edge1 - EdgeUV1.Y * Edge2).Size();
+				UVLength.Y = (-EdgeUV2.X * Edge1 + EdgeUV1.X * Edge2).Size();
+
+				Chart.WorldScale += UVLength;
+				Chart.UVArea += UVArea;
+			}
+
+			Chart.LastTri = Tri;
 
 #if !CHART_JOINING
-		if (LayoutVersion >= FMeshDescriptionOperations::ELightmapUVVersion::SmallChartPacking)
-		{
-			Chart.WorldScale /= FMath::Max(Chart.UVArea, 1e-8f);
-		}
-		else
-		{
-			if (Chart.UVArea > 1e-4f)
+			if (LayoutVersion >= FMeshDescriptionOperations::ELightmapUVVersion::SmallChartPacking)
 			{
-				Chart.WorldScale /= Chart.UVArea;
+				Chart.WorldScale /= FMath::Max(Chart.UVArea, 1e-8f);
 			}
 			else
 			{
-				Chart.WorldScale = FVector2D::ZeroVector;
+				if (Chart.UVArea > 1e-4f)
+				{
+					Chart.WorldScale /= Chart.UVArea;
+				}
+				else
+				{
+					Chart.WorldScale = FVector2D::ZeroVector;
+				}
 			}
-		}		
 
-		TotalUVArea += Chart.UVArea * Chart.WorldScale.X * Chart.WorldScale.Y;
+			TotalUVArea += Chart.UVArea * Chart.WorldScale.X * Chart.WorldScale.Y;
 #endif
-	}
+		}
 
 #if CHART_JOINING
-	for( int32 i = 0; i < Charts.Num(); i++ )
-	{
-		FMeshChart& Chart = Charts[i];
-
-		for( uint32 Tri = Chart.FirstTri; Tri < Chart.LastTri; Tri++ )
+		for (int32 i = 0; i < Charts.Num(); i++)
 		{
-			for( int k = 0; k < 3; k++ )
-			{
-				uint32 Index = 3 * SortedTris[ Tri ] + k;
+			FMeshChart& Chart = Charts[i];
 
-				if( TranslatedMatches[ Index ] >= 0 )
+			for (uint32 Tri = Chart.FirstTri; Tri < Chart.LastTri; Tri++)
+			{
+				for (int k = 0; k < 3; k++)
 				{
-					checkSlow( TranslatedMatches[ TranslatedMatches[ Index ] ] == Index );
-
-					uint32 V0i = Index;
-					uint32 V0j = TranslatedMatches[ Index ];
+					uint32 Index = 3 * SortedTris[Tri] + k;
 
-					uint32 TriI = V0i / 3;
-					uint32 TriJ = V0j / 3;
-					
-					if( TriJ <= TriI )
+					if (TranslatedMatches[Index] >= 0)
 					{
-						// Only need to consider one direction
-						continue;
-					}
-					
-					uint32 V1i = 3 * TriI + (V0i + 1) % 3;
-					uint32 V1j = 3 * TriJ + (V0j + 1) % 3;
+						checkSlow(TranslatedMatches[TranslatedMatches[Index]] == Index);
 
-					int32 ChartI = i;
-					int32 ChartJ = DisjointSetToChartMap[ DisjointSet[ TriJ ] ];
+						uint32 V0i = Index;
+						uint32 V0j = TranslatedMatches[Index];
 
-					FVector2D UV0i = TexCoords[ V0i ];
-					FVector2D UV1i = TexCoords[ V1i ];
-					FVector2D UV0j = TexCoords[ V0j ];
-					FVector2D UV1j = TexCoords[ V1j ];
+						uint32 TriI = V0i / 3;
+						uint32 TriJ = V0j / 3;
 
-					FVector2D EdgeUVi = UV1i - UV0i;
-					FVector2D EdgeUVj = UV1j - UV0j;
+						if (TriJ <= TriI)
+						{
+							// Only need to consider one direction
+							continue;
+						}
 
-					bool bMirrored = TriangleUVArea( TriI ) * TriangleUVArea( TriJ ) < 0.0f;
-					
-					FVector2D EdgeOffset0 = UV0i - UV1j;
-					FVector2D EdgeOffset1 = UV1i - UV0j;
-					checkSlow( ( EdgeOffset0 - EdgeOffset1 ).IsNearlyZero( ThreshUVsAreSame ) );
+						uint32 V1i = 3 * TriI + (V0i + 1) % 3;
+						uint32 V1j = 3 * TriJ + (V0j + 1) % 3;
 
-					FVector2D Translation = EdgeOffset0;
+						int32 ChartI = i;
+						int32 ChartJ = DisjointSetToChartMap[DisjointSet[TriJ]];
 
-					FMeshChart& ChartA = Charts[ ChartI ];
-					FMeshChart& ChartB = Charts[ ChartJ ];
+						FVector2D UV0i = TexCoords[V0i];
+						FVector2D UV1i = TexCoords[V1i];
+						FVector2D UV0j = TexCoords[V0j];
+						FVector2D UV1j = TexCoords[V1j];
 
-					for( uint32 Side = 0; Side < 4; Side++ )
-					{
-						// Join[] = { left, right, bottom, top }
+						FVector2D EdgeUVi = UV1i - UV0i;
+						FVector2D EdgeUVj = UV1j - UV0j;
 
-						// FIXME
-						if( bMirrored )
-							continue;
+						bool bMirrored = TriangleUVArea(TriI) * TriangleUVArea(TriJ) < 0.0f;
 
-						if( ChartA.Join[ Side ^ 0 ] != -1 ||
-							ChartB.Join[ Side ^ 1 ] != -1 )
-						{
-							// Already joined with something else
-							continue;
-						}
+						FVector2D EdgeOffset0 = UV0i - UV1j;
+						FVector2D EdgeOffset1 = UV1i - UV0j;
+						checkSlow((EdgeOffset0 - EdgeOffset1).IsNearlyZero(ThreshUVsAreSame));
 
-						uint32 Sign = Side & 1;
-						uint32 Axis = Side >> 1;
+						FVector2D Translation = EdgeOffset0;
 
-						bool bAxisAligned = FMath::Abs( EdgeUVi[ Axis ] ) < ThreshUVsAreSame;
-						bool bBorderA = FMath::Abs( UV0i[ Axis ] - ( Sign ^ 0 ? Chart.MaxUV[ Axis ] : Chart.MinUV[ Axis ] ) ) < ThreshUVsAreSame;
-						bool bBorderB = FMath::Abs( UV0j[ Axis ] - ( Sign ^ 1 ? Chart.MaxUV[ Axis ] : Chart.MinUV[ Axis ] ) ) < ThreshUVsAreSame;
+						FMeshChart& ChartA = Charts[ChartI];
+						FMeshChart& ChartB = Charts[ChartJ];
 
-						// FIXME mirrored
-						if( !bAxisAligned || !bBorderA || !bBorderB )
+						for (uint32 Side = 0; Side < 4; Side++)
 						{
-							// Edges weren't on matching rectangle borders
-							continue;
-						}
+							// Join[] = { left, right, bottom, top }
 
-						FVector2D CenterA = 0.5f * ( ChartA.MinUV + ChartA.MaxUV );
-						FVector2D CenterB = 0.5f * ( ChartB.MinUV + ChartB.MaxUV );
+							// FIXME
+							if (bMirrored)
+								continue;
 
-						FVector2D ExtentA = 0.5f * ( ChartA.MaxUV - ChartA.MinUV );
-						FVector2D ExtentB = 0.5f * ( ChartB.MaxUV - ChartB.MinUV );
+							if (ChartA.Join[Side ^ 0] != -1 ||
+								ChartB.Join[Side ^ 1] != -1)
+							{
+								// Already joined with something else
+								continue;
+							}
 
-						// FIXME mirrored
-						CenterB += Translation;
+							uint32 Sign = Side & 1;
+							uint32 Axis = Side >> 1;
 
-						FVector2D CenterDiff = CenterA - CenterB;
-						FVector2D ExtentDiff = ExtentA - ExtentB;
-						FVector2D Separation = ExtentA + ExtentB + CenterDiff * ( Sign ? 1.0f : -1.0f );
+							bool bAxisAligned = FMath::Abs(EdgeUVi[Axis]) < ThreshUVsAreSame;
+							bool bBorderA = FMath::Abs(UV0i[Axis] - (Sign ^ 0 ? Chart.MaxUV[Axis] : Chart.MinUV[Axis])) < ThreshUVsAreSame;
+							bool bBorderB = FMath::Abs(UV0j[Axis] - (Sign ^ 1 ? Chart.MaxUV[Axis] : Chart.MinUV[Axis])) < ThreshUVsAreSame;
 
-						bool bCenterMatch = FMath::Abs( CenterDiff[ Axis ^ 1 ] ) < ThreshUVsAreSame;
-						bool bExtentMatch = FMath::Abs( ExtentDiff[ Axis ^ 1 ] ) < ThreshUVsAreSame;
-						bool bSeparate    = FMath::Abs( Separation[ Axis ^ 0 ] ) < ThreshUVsAreSame;
-	
-						if( !bCenterMatch || !bExtentMatch || !bSeparate )
-						{
-							// Rectangles don't match up after translation
-							continue;
-						}
+							// FIXME mirrored
+							if (!bAxisAligned || !bBorderA || !bBorderB)
+							{
+								// Edges weren't on matching rectangle borders
+								continue;
+							}
 
-						// Found a valid edge join
-						ChartA.Join[ Side ^ 0 ] = ChartJ;
-						ChartB.Join[ Side ^ 1 ] = ChartI;
-						break;
+							FVector2D CenterA = 0.5f * (ChartA.MinUV + ChartA.MaxUV);
+							FVector2D CenterB = 0.5f * (ChartB.MinUV + ChartB.MaxUV);
+
+							FVector2D ExtentA = 0.5f * (ChartA.MaxUV - ChartA.MinUV);
+							FVector2D ExtentB = 0.5f * (ChartB.MaxUV - ChartB.MinUV);
+
+							// FIXME mirrored
+							CenterB += Translation;
+
+							FVector2D CenterDiff = CenterA - CenterB;
+							FVector2D ExtentDiff = ExtentA - ExtentB;
+							FVector2D Separation = ExtentA + ExtentB + CenterDiff * (Sign ? 1.0f : -1.0f);
+
+							bool bCenterMatch = FMath::Abs(CenterDiff[Axis ^ 1]) < ThreshUVsAreSame;
+							bool bExtentMatch = FMath::Abs(ExtentDiff[Axis ^ 1]) < ThreshUVsAreSame;
+							bool bSeparate = FMath::Abs(Separation[Axis ^ 0]) < ThreshUVsAreSame;
+
+							if (!bCenterMatch || !bExtentMatch || !bSeparate)
+							{
+								// Rectangles don't match up after translation
+								continue;
+							}
+
+							// Found a valid edge join
+							ChartA.Join[Side ^ 0] = ChartJ;
+							ChartB.Join[Side ^ 1] = ChartI;
+							break;
+						}
 					}
 				}
 			}
 		}
-	}
-
-	TArray< uint32 > JoinedSortedTris;
-	JoinedSortedTris.Reserve( NumTris );
 
-	// Detect loops
-	for( uint32 Axis = 0; Axis < 2; Axis++ )
-	{
-		uint32 Side = Axis << 1;
+		TArray< uint32 > JoinedSortedTris;
+		JoinedSortedTris.Reserve(NumTris);
 
-		for( int32 i = 0; i < Charts.Num(); i++ )
+		// Detect loops
+		for (uint32 Axis = 0; Axis < 2; Axis++)
 		{
-			int32 j = Charts[i].Join[ Side ^ 1 ];
-			while( j != -1 )
+			uint32 Side = Axis << 1;
+
+			for (int32 i = 0; i < Charts.Num(); i++)
 			{
-				int32 Next = Charts[j].Join[ Side ^ 1 ];
-				if( Next == i )
+				int32 j = Charts[i].Join[Side ^ 1];
+				while (j != -1)
 				{
-					// Break loop
-					Charts[i].Join[ Side ^ 0 ] = -1;
-					Charts[j].Join[ Side ^ 1 ] = -1;
-					break;
+					int32 Next = Charts[j].Join[Side ^ 1];
+					if (Next == i)
+					{
+						// Break loop
+						Charts[i].Join[Side ^ 0] = -1;
+						Charts[j].Join[Side ^ 1] = -1;
+						break;
+					}
+					j = Next;
 				}
-				j = Next;
 			}
 		}
-	}
 
-	// Join rows first, then columns
-	for( uint32 Axis = 0; Axis < 2; Axis++ )
-	{
-		for( int32 i = 0; i < Charts.Num(); i++ )
+		// Join rows first, then columns
+		for (uint32 Axis = 0; Axis < 2; Axis++)
 		{
-			FMeshChart& ChartA = Charts[i];
-
-			if( ChartA.FirstTri == ChartA.LastTri )
-			{
-				// Empty chart
-				continue;
-			}
-			
-			for( uint32 Side = 0; Side < 4; Side++ )
+			for (int32 i = 0; i < Charts.Num(); i++)
 			{
-				if( ChartA.Join[ Side ] != -1 )
-				{
-					FMeshChart& ChartB = Charts[ ChartA.Join[ Side ] ];
+				FMeshChart& ChartA = Charts[i];
 
-					check( ChartB.Join[ Side ^ 1 ] == i );
-					check( ChartB.FirstTri != ChartB.LastTri );
+				if (ChartA.FirstTri == ChartA.LastTri)
+				{
+					// Empty chart
+					continue;
 				}
-			}
-		}
 
-		NumTris = 0;
+				for (uint32 Side = 0; Side < 4; Side++)
+				{
+					if (ChartA.Join[Side] != -1)
+					{
+						FMeshChart& ChartB = Charts[ChartA.Join[Side]];
 
-		for( int32 i = 0; i < Charts.Num(); i++ )
-		{
-			FMeshChart& Chart = Charts[i];
-			NumTris += Chart.LastTri - Chart.FirstTri;
-		}
-		check( NumTris == SortedTris.Num() );
+						check(ChartB.Join[Side ^ 1] == i);
+						check(ChartB.FirstTri != ChartB.LastTri);
+					}
+				}
+			}
 
-		NumTris = 0;
-		for( int32 i = 0; i < Charts.Num(); i++ )
-		{
-			FMeshChart& ChartA = Charts[i];
+			NumTris = 0;
 
-			if( ChartA.FirstTri == ChartA.LastTri )
+			for (int32 i = 0; i < Charts.Num(); i++)
 			{
-				// Empty chart
-				continue;
+				FMeshChart& Chart = Charts[i];
+				NumTris += Chart.LastTri - Chart.FirstTri;
 			}
+			check(NumTris == SortedTris.Num());
 
-			uint32 Side = Axis << 1;
-
-			// Find start (left, bottom)
-			if( ChartA.Join[ Side ^ 0 ] == -1 )
+			NumTris = 0;
+			for (int32 i = 0; i < Charts.Num(); i++)
 			{
-				// Add original tris
-				NumTris += ChartA.LastTri - ChartA.FirstTri;
+				FMeshChart& ChartA = Charts[i];
 
-				// Continue joining until no more to the (right, top)
-				int32 Next = ChartA.Join[ Side ^ 1 ];
-				while( Next != -1 )
+				if (ChartA.FirstTri == ChartA.LastTri)
 				{
-					FMeshChart& ChartB = Charts[ Next ];
-
-					NumTris += ChartB.LastTri - ChartB.FirstTri;
-					Next = ChartB.Join[ Side ^ 1 ];
+					// Empty chart
+					continue;
 				}
-			}
-		}
-		check( NumTris == SortedTris.Num() );
 
-#if 1
-		NumTris = 0;
-		for( int32 i = 0; i < Charts.Num(); i++ )
-		{
-			FMeshChart& ChartA = Charts[i];
+				uint32 Side = Axis << 1;
 
-			if( ChartA.FirstTri == ChartA.LastTri )
-			{
-				// Empty chart
-				continue;
-			}
+				// Find start (left, bottom)
+				if (ChartA.Join[Side ^ 0] == -1)
+				{
+					// Add original tris
+					NumTris += ChartA.LastTri - ChartA.FirstTri;
 
-			// Join[] = { left, right, bottom, top }
+					// Continue joining until no more to the (right, top)
+					int32 Next = ChartA.Join[Side ^ 1];
+					while (Next != -1)
+					{
+						FMeshChart& ChartB = Charts[Next];
 
-			uint32 Side = Axis << 1;
+						NumTris += ChartB.LastTri - ChartB.FirstTri;
+						Next = ChartB.Join[Side ^ 1];
+					}
+				}
+			}
+			check(NumTris == SortedTris.Num());
 
-			// Find start (left, bottom)
-			if( ChartA.Join[ Side ^ 0 ] == -1 )
+#if 1
+			NumTris = 0;
+			for (int32 i = 0; i < Charts.Num(); i++)
 			{
-				uint32 FirstTri = JoinedSortedTris.Num();
-			
-				// Add original tris
-				for( uint32 Tri = ChartA.FirstTri; Tri < ChartA.LastTri; Tri++ )
+				FMeshChart& ChartA = Charts[i];
+
+				if (ChartA.FirstTri == ChartA.LastTri)
 				{
-					JoinedSortedTris.Add( SortedTris[ Tri ] );
+					// Empty chart
+					continue;
 				}
-				NumTris += ChartA.LastTri - ChartA.FirstTri;
 
-				// Continue joining until no more to the (right, top)
-				while( ChartA.Join[ Side ^ 1 ] != -1 )
-				{
-					FMeshChart& ChartB = Charts[ ChartA.Join[ Side ^ 1 ] ];
+				// Join[] = { left, right, bottom, top }
+
+				uint32 Side = Axis << 1;
 
-					check( ChartB.FirstTri != ChartB.LastTri );
+				// Find start (left, bottom)
+				if (ChartA.Join[Side ^ 0] == -1)
+				{
+					uint32 FirstTri = JoinedSortedTris.Num();
 
-					FVector2D Translation = ChartA.MinUV - ChartB.MinUV;
-					Translation[ Axis ] += ChartA.MaxUV[ Axis ] - ChartA.MinUV[ Axis ];
+					// Add original tris
+					for (uint32 Tri = ChartA.FirstTri; Tri < ChartA.LastTri; Tri++)
+					{
+						JoinedSortedTris.Add(SortedTris[Tri]);
+					}
+					NumTris += ChartA.LastTri - ChartA.FirstTri;
 
-					for( uint32 Tri = ChartB.FirstTri; Tri < ChartB.LastTri; Tri++ )
+					// Continue joining until no more to the (right, top)
+					while (ChartA.Join[Side ^ 1] != -1)
 					{
-						JoinedSortedTris.Add( SortedTris[ Tri ] );
-						for( int k = 0; k < 3; k++ )
+						FMeshChart& ChartB = Charts[ChartA.Join[Side ^ 1]];
+
+						check(ChartB.FirstTri != ChartB.LastTri);
+
+						FVector2D Translation = ChartA.MinUV - ChartB.MinUV;
+						Translation[Axis] += ChartA.MaxUV[Axis] - ChartA.MinUV[Axis];
+
+						for (uint32 Tri = ChartB.FirstTri; Tri < ChartB.LastTri; Tri++)
 						{
-							TexCoords[ 3 * SortedTris[ Tri ] + k ] += Translation;
+							JoinedSortedTris.Add(SortedTris[Tri]);
+							for (int k = 0; k < 3; k++)
+							{
+								TexCoords[3 * SortedTris[Tri] + k] += Translation;
+							}
 						}
-					}
-					NumTris += ChartB.LastTri - ChartB.FirstTri;
+						NumTris += ChartB.LastTri - ChartB.FirstTri;
 
-					ChartA.Join[ Side ^ 1 ] = ChartB.Join[ Side ^ 1 ];
-					ChartA.MaxUV[ Axis ] += ChartB.MaxUV[ Axis ] - ChartB.MinUV[ Axis ];
-					ChartA.WorldScale += ChartB.WorldScale;
-					ChartA.UVArea += ChartB.UVArea;
+						ChartA.Join[Side ^ 1] = ChartB.Join[Side ^ 1];
+						ChartA.MaxUV[Axis] += ChartB.MaxUV[Axis] - ChartB.MinUV[Axis];
+						ChartA.WorldScale += ChartB.WorldScale;
+						ChartA.UVArea += ChartB.UVArea;
 
-					ChartB.FirstTri = 0;
-					ChartB.LastTri = 0;
-					ChartB.UVArea = 0.0f;
+						ChartB.FirstTri = 0;
+						ChartB.LastTri = 0;
+						ChartB.UVArea = 0.0f;
 
-					DisconnectChart( ChartB, Side ^ 2 );
-					DisconnectChart( ChartB, Side ^ 3 );
-				}
+						DisconnectChart(ChartB, Side ^ 2);
+						DisconnectChart(ChartB, Side ^ 3);
+					}
 
-				ChartA.FirstTri = FirstTri;
-				ChartA.LastTri = JoinedSortedTris.Num();
-			}
-			else
-			{
-				// Make sure a starting chart could connect to this
-				FMeshChart& ChartB = Charts[ ChartA.Join[ Side ^ 0 ] ];
-				check( ChartB.Join[ Side ^ 1 ] == i );
-				check( ChartB.FirstTri != ChartB.LastTri );
+					ChartA.FirstTri = FirstTri;
+					ChartA.LastTri = JoinedSortedTris.Num();
+				}
+				else
+				{
+					// Make sure a starting chart could connect to this
+					FMeshChart& ChartB = Charts[ChartA.Join[Side ^ 0]];
+					check(ChartB.Join[Side ^ 1] == i);
+					check(ChartB.FirstTri != ChartB.LastTri);
+				}
 			}
-		}
-		check( NumTris == SortedTris.Num() );
+			check(NumTris == SortedTris.Num());
 
-		check( SortedTris.Num() == JoinedSortedTris.Num() );
-		Exchange( SortedTris, JoinedSortedTris );
-		JoinedSortedTris.Reset();
+			check(SortedTris.Num() == JoinedSortedTris.Num());
+			Exchange(SortedTris, JoinedSortedTris);
+			JoinedSortedTris.Reset();
 #endif
-	}
-	
-	// Clean out empty charts
-	for( int32 i = 0; i < Charts.Num(); i++ )
-	{
-		while( i < Charts.Num() && Charts[i].FirstTri == Charts[i].LastTri )
-		{
-			Charts.RemoveAtSwap(i);
 		}
-	}
 
-	for( int32 i = 0; i < Charts.Num(); i++ )
-	{
-		FMeshChart& Chart = Charts[i];
-
-		if (LayoutVersion >= FMeshDescriptionOperations::ELightmapUVVersion::SmallChartPacking)
+		// Clean out empty charts
+		for (int32 i = 0; i < Charts.Num(); i++)
 		{
-			Chart.WorldScale /= FMath::Max(Chart.UVArea, 1e-8f);
+			while (i < Charts.Num() && Charts[i].FirstTri == Charts[i].LastTri)
+			{
+				Charts.RemoveAtSwap(i);
+			}
 		}
-		else
+
+		for (int32 i = 0; i < Charts.Num(); i++)
 		{
-			if (Chart.UVArea > 1e-4f)
+			FMeshChart& Chart = Charts[i];
+
+			if (LayoutVersion >= FMeshDescriptionOperations::ELightmapUVVersion::SmallChartPacking)
 			{
-				Chart.WorldScale /= Chart.UVArea;
+				Chart.WorldScale /= FMath::Max(Chart.UVArea, 1e-8f);
 			}
 			else
 			{
-				Chart.WorldScale = FVector2D::ZeroVector;
+				if (Chart.UVArea > 1e-4f)
+				{
+					Chart.WorldScale /= Chart.UVArea;
+				}
+				else
+				{
+					Chart.WorldScale = FVector2D::ZeroVector;
+				}
 			}
-		}
 
-		TotalUVArea += Chart.UVArea * Chart.WorldScale.X * Chart.WorldScale.Y;
-	}
+			TotalUVArea += Chart.UVArea * Chart.WorldScale.X * Chart.WorldScale.Y;
+		}
 #endif
 
-	double End = FPlatformTime::Seconds();
+		double End = FPlatformTime::Seconds();
 
-	UE_LOG(LogMeshDescriptionLayoutUV, Display, TEXT("FindCharts: %s"), *FPlatformTime::PrettyTime(End - Begin) );
-}
-
-bool FLayoutUV::FindBestPacking()
-{
-	if( (uint32)Charts.Num() > TextureResolution * TextureResolution || TotalUVArea == 0.f )
-	{
-		// More charts than texels
-		return false;
+		UE_LOG(LogMeshDescriptionLayoutUV, Display, TEXT("FindCharts: %s"), *FPlatformTime::PrettyTime(End - Begin));
 	}
-	
-	const float LinearSearchStart = 0.5f;
-	const float LinearSearchStep = 0.5f;
-	const int32 BinarySearchSteps = 6;
-	
-	float UVScaleFail = TextureResolution * FMath::Sqrt( 1.0f / TotalUVArea );
-	float UVScalePass = TextureResolution * FMath::Sqrt( LinearSearchStart / TotalUVArea );
-
-	// Linear search for first fit
-	while(1)
-	{
-		ScaleCharts( UVScalePass );
 
-		bool bFit = PackCharts();
-		if( bFit )
+	bool FLayoutUV::FindBestPacking()
+	{
+		if ((uint32)Charts.Num() > TextureResolution * TextureResolution || TotalUVArea == 0.f)
 		{
-			break;
+			// More charts than texels
+			return false;
 		}
 
-		UVScaleFail = UVScalePass;
-		UVScalePass *= LinearSearchStep;
-	}
+		const float LinearSearchStart = 0.5f;
+		const float LinearSearchStep = 0.5f;
+		const int32 BinarySearchSteps = 6;
 
-	// Binary search for best fit
-	for( int32 i = 0; i < BinarySearchSteps; i++ )
-	{
-		float UVScale = 0.5f * ( UVScaleFail + UVScalePass );
-		ScaleCharts( UVScale );
+		float UVScaleFail = TextureResolution * FMath::Sqrt(1.0f / TotalUVArea);
+		float UVScalePass = TextureResolution * FMath::Sqrt(LinearSearchStart / TotalUVArea);
 
-		bool bFit = PackCharts();
-		if( bFit )
-		{
-			UVScalePass = UVScale;
-		}
-		else
+		// Linear search for first fit
+		while (1)
 		{
-			UVScaleFail = UVScale;
-		}
-	}
+			ScaleCharts(UVScalePass);
 
-	// TODO store packing scale/bias separate so this isn't necessary
-	ScaleCharts( UVScalePass );
-	PackCharts();
+			bool bFit = PackCharts();
+			if (bFit)
+			{
+				break;
+			}
 
-	return true;
-}
+			UVScaleFail = UVScalePass;
+			UVScalePass *= LinearSearchStep;
+		}
 
-void FLayoutUV::ScaleCharts( float UVScale )
-{
-	for( int32 i = 0; i < Charts.Num(); i++ )
-	{
-		FMeshChart& Chart = Charts[i];
-		Chart.UVScale = Chart.WorldScale * UVScale;
-	}
-	
-	// Scale charts such that they all fit and roughly total the same area as before
-#if 1
-	float UniformScale = 1.0f;
-	for( int i = 0; i < 1000; i++ )
-	{
-		uint32 NumMaxedOut = 0;
-		float ScaledUVArea = 0.0f;
-		for( int32 ChartIndex = 0; ChartIndex < Charts.Num(); ChartIndex++ )
+		// Binary search for best fit
+		for (int32 i = 0; i < BinarySearchSteps; i++)
 		{
-			FMeshChart& Chart = Charts[ChartIndex];
-
-			FVector2D ChartSize	= Chart.MaxUV - Chart.MinUV;
-			FVector2D ChartSizeScaled = ChartSize * Chart.UVScale * UniformScale;
-
-			const float MaxChartEdge = TextureResolution - 1.0f;
-			const float LongestChartEdge = FMath::Max( ChartSizeScaled.X, ChartSizeScaled.Y );
+			float UVScale = 0.5f * (UVScaleFail + UVScalePass);
+			ScaleCharts(UVScale);
 
-			const float Epsilon = 0.01f;
-			if( LongestChartEdge + Epsilon > MaxChartEdge )
+			bool bFit = PackCharts();
+			if (bFit)
 			{
-				// Rescale oversized charts to fit
-				Chart.UVScale.X = MaxChartEdge / FMath::Max( ChartSize.X, ChartSize.Y );
-				Chart.UVScale.Y = MaxChartEdge / FMath::Max( ChartSize.X, ChartSize.Y );
-				NumMaxedOut++;
+				UVScalePass = UVScale;
 			}
 			else
 			{
-				Chart.UVScale.X *= UniformScale;
-				Chart.UVScale.Y *= UniformScale;
+				UVScaleFail = UVScale;
 			}
-			
-			ScaledUVArea += Chart.UVArea * Chart.UVScale.X * Chart.UVScale.Y;
 		}
 
-		if( NumMaxedOut == 0 )
-		{
-			// No charts maxed out so no need to rebalance
-			break;
-		}
+		// TODO store packing scale/bias separate so this isn't necessary
+		ScaleCharts(UVScalePass);
+		PackCharts();
 
-		if( NumMaxedOut == Charts.Num() )
-		{
-			// All charts are maxed out
-			break;
-		}
+		return true;
+	}
 
-		// Scale up smaller charts to maintain expected total area
-		// Want ScaledUVArea == TotalUVArea * UVScale^2
-		float RebalanceScale = UVScale * FMath::Sqrt( TotalUVArea / ScaledUVArea );
-		if( RebalanceScale < 1.01f )
+	void FLayoutUV::ScaleCharts(float UVScale)
+	{
+		for (int32 i = 0; i < Charts.Num(); i++)
 		{
-			// Stop if further rebalancing is minor
-			break;
+			FMeshChart& Chart = Charts[i];
+			Chart.UVScale = Chart.WorldScale * UVScale;
 		}
-		UniformScale = RebalanceScale;
-	}
-#endif
 
+		// Scale charts such that they all fit and roughly total the same area as before
 #if 1
-	float NonuniformScale = 1.0f;
-	for( int i = 0; i < 1000; i++ )
-	{
-		uint32 NumMaxedOut = 0;
-		float ScaledUVArea = 0.0f;
-		for( int32 ChartIndex = 0; ChartIndex < Charts.Num(); ChartIndex++ )
+		float UniformScale = 1.0f;
+		for (int i = 0; i < 1000; i++)
 		{
-			FMeshChart& Chart = Charts[ChartIndex];
-
-			for( int k = 0; k < 2; k++ )
+			uint32 NumMaxedOut = 0;
+			float ScaledUVArea = 0.0f;
+			for (int32 ChartIndex = 0; ChartIndex < Charts.Num(); ChartIndex++)
 			{
-				const float MaximumChartSize = TextureResolution - 1.0f;
-				const float ChartSize = Chart.MaxUV[k] - Chart.MinUV[k];
-				const float ChartSizeScaled = ChartSize * Chart.UVScale[k] * NonuniformScale;
+				FMeshChart& Chart = Charts[ChartIndex];
+
+				FVector2D ChartSize = Chart.MaxUV - Chart.MinUV;
+				FVector2D ChartSizeScaled = ChartSize * Chart.UVScale * UniformScale;
+
+				const float MaxChartEdge = TextureResolution - 1.0f;
+				const float LongestChartEdge = FMath::Max(ChartSizeScaled.X, ChartSizeScaled.Y);
 
 				const float Epsilon = 0.01f;
-				if( ChartSizeScaled + Epsilon > MaximumChartSize )
+				if (LongestChartEdge + Epsilon > MaxChartEdge)
 				{
-					// Scale oversized charts to max size
-					Chart.UVScale[k] = MaximumChartSize / ChartSize;
+					// Rescale oversized charts to fit
+					Chart.UVScale.X = MaxChartEdge / FMath::Max(ChartSize.X, ChartSize.Y);
+					Chart.UVScale.Y = MaxChartEdge / FMath::Max(ChartSize.X, ChartSize.Y);
 					NumMaxedOut++;
 				}
 				else
 				{
-					Chart.UVScale[k] *= NonuniformScale;
+					Chart.UVScale.X *= UniformScale;
+					Chart.UVScale.Y *= UniformScale;
 				}
+
+				ScaledUVArea += Chart.UVArea * Chart.UVScale.X * Chart.UVScale.Y;
 			}
 
-			ScaledUVArea += Chart.UVArea * Chart.UVScale.X * Chart.UVScale.Y;
-		}
+			if (NumMaxedOut == 0)
+			{
+				// No charts maxed out so no need to rebalance
+				break;
+			}
 
-		if( NumMaxedOut == 0 )
-		{
-			// No charts maxed out so no need to rebalance
-			break;
+			if (NumMaxedOut == Charts.Num())
+			{
+				// All charts are maxed out
+				break;
+			}
+
+			// Scale up smaller charts to maintain expected total area
+			// Want ScaledUVArea == TotalUVArea * UVScale^2
+			float RebalanceScale = UVScale * FMath::Sqrt(TotalUVArea / ScaledUVArea);
+			if (RebalanceScale < 1.01f)
+			{
+				// Stop if further rebalancing is minor
+				break;
+			}
+			UniformScale = RebalanceScale;
 		}
+#endif
 
-		if( NumMaxedOut == Charts.Num() * 2 )
+#if 1
+		float NonuniformScale = 1.0f;
+		for (int i = 0; i < 1000; i++)
 		{
-			// All charts are maxed out in both dimensions
-			break;
+			uint32 NumMaxedOut = 0;
+			float ScaledUVArea = 0.0f;
+			for (int32 ChartIndex = 0; ChartIndex < Charts.Num(); ChartIndex++)
+			{
+				FMeshChart& Chart = Charts[ChartIndex];
+
+				for (int k = 0; k < 2; k++)
+				{
+					const float MaximumChartSize = TextureResolution - 1.0f;
+					const float ChartSize = Chart.MaxUV[k] - Chart.MinUV[k];
+					const float ChartSizeScaled = ChartSize * Chart.UVScale[k] * NonuniformScale;
+
+					const float Epsilon = 0.01f;
+					if (ChartSizeScaled + Epsilon > MaximumChartSize)
+					{
+						// Scale oversized charts to max size
+						Chart.UVScale[k] = MaximumChartSize / ChartSize;
+						NumMaxedOut++;
+					}
+					else
+					{
+						Chart.UVScale[k] *= NonuniformScale;
+					}
+				}
+
+				ScaledUVArea += Chart.UVArea * Chart.UVScale.X * Chart.UVScale.Y;
+			}
+
+			if (NumMaxedOut == 0)
+			{
+				// No charts maxed out so no need to rebalance
+				break;
+			}
+
+			if (NumMaxedOut == Charts.Num() * 2)
+			{
+				// All charts are maxed out in both dimensions
+				break;
+			}
+
+			// Scale up smaller charts to maintain expected total area
+			// Want ScaledUVArea == TotalUVArea * UVScale^2
+			float RebalanceScale = UVScale * FMath::Sqrt(TotalUVArea / ScaledUVArea);
+			if (RebalanceScale < 1.01f)
+			{
+				// Stop if further rebalancing is minor
+				break;
+			}
+			NonuniformScale = RebalanceScale;
 		}
+#endif
 
-		// Scale up smaller charts to maintain expected total area
-		// Want ScaledUVArea == TotalUVArea * UVScale^2
-		float RebalanceScale = UVScale * FMath::Sqrt( TotalUVArea / ScaledUVArea );
-		if( RebalanceScale < 1.01f )
+		// Sort charts from largest to smallest
+		struct FCompareCharts
 		{
-			// Stop if further rebalancing is minor
-			break;
-		}
-		NonuniformScale = RebalanceScale;
+			FORCEINLINE bool operator()(const FMeshChart& A, const FMeshChart& B) const
+			{
+				// Rect area
+				FVector2D ChartRectA = (A.MaxUV - A.MinUV) * A.UVScale;
+				FVector2D ChartRectB = (B.MaxUV - B.MinUV) * B.UVScale;
+				return ChartRectA.X * ChartRectA.Y > ChartRectB.X * ChartRectB.Y;
+			}
+		};
+		Algo::IntroSort(Charts, FCompareCharts());
 	}
-#endif
 
-	// Sort charts from largest to smallest
-	struct FCompareCharts
+	bool FLayoutUV::PackCharts()
 	{
-		FORCEINLINE bool operator()( const FMeshChart& A, const FMeshChart& B ) const
-		{
-			// Rect area
-			FVector2D ChartRectA = ( A.MaxUV - A.MinUV ) * A.UVScale;
-			FVector2D ChartRectB = ( B.MaxUV - B.MinUV ) * B.UVScale;
-			return ChartRectA.X * ChartRectA.Y > ChartRectB.X * ChartRectB.Y;
-		}
-	};
-	Algo::IntroSort( Charts, FCompareCharts() );
-}
+		uint32 RasterizeCycles = 0;
+		uint32 FindCycles = 0;
 
-bool FLayoutUV::PackCharts()
-{
-	uint32 RasterizeCycles = 0;
-	uint32 FindCycles = 0;
-	
-	double BeginPackCharts = FPlatformTime::Seconds();
+		double BeginPackCharts = FPlatformTime::Seconds();
 
-	LayoutRaster.Clear();
+		LayoutRaster.Clear();
 
-	for( int32 i = 0; i < Charts.Num(); i++ )
-	{
-		FMeshChart& Chart = Charts[i];
+		for (int32 i = 0; i < Charts.Num(); i++)
+		{
+			FMeshChart& Chart = Charts[i];
 
-		// Try different orientations and pick best
-		int32				BestOrientation = -1;
-		FAllocator2D::FRect	BestRect = { ~0u, ~0u, ~0u, ~0u };
+			// Try different orientations and pick best
+			int32				BestOrientation = -1;
+			FAllocator2D::FRect	BestRect = { ~0u, ~0u, ~0u, ~0u };
 
-		for( int32 Orientation = 0; Orientation < 8; Orientation++ )
-		{
-			// TODO If any dimension is less than 1 pixel shrink dimension to zero
+			for (int32 Orientation = 0; Orientation < 8; Orientation++)
+			{
+				// TODO If any dimension is less than 1 pixel shrink dimension to zero
 
-			OrientChart( Chart, Orientation );
-			
-			FVector2D ChartSize = Chart.MaxUV - Chart.MinUV;
-			ChartSize = ChartSize.X * Chart.PackingScaleU + ChartSize.Y * Chart.PackingScaleV;
+				OrientChart(Chart, Orientation);
 
-			// Only need half pixel dilate for rects
-			FAllocator2D::FRect	Rect;
-			Rect.X = 0;
-			Rect.Y = 0;
-			Rect.W = FMath::CeilToInt( FMath::Abs( ChartSize.X ) + 1.0f );
-			Rect.H = FMath::CeilToInt( FMath::Abs( ChartSize.Y ) + 1.0f );
+				FVector2D ChartSize = Chart.MaxUV - Chart.MinUV;
+				ChartSize = ChartSize.X * Chart.PackingScaleU + ChartSize.Y * Chart.PackingScaleV;
 
-			// Just in case lack of precision pushes it over
-			Rect.W = FMath::Min( TextureResolution, Rect.W );
-			Rect.H = FMath::Min( TextureResolution, Rect.H );
+				// Only need half pixel dilate for rects
+				FAllocator2D::FRect	Rect;
+				Rect.X = 0;
+				Rect.Y = 0;
+				Rect.W = FMath::CeilToInt(FMath::Abs(ChartSize.X) + 1.0f);
+				Rect.H = FMath::CeilToInt(FMath::Abs(ChartSize.Y) + 1.0f);
 
-			const bool bRectPack = false;
+				// Just in case lack of precision pushes it over
+				Rect.W = FMath::Min(TextureResolution, Rect.W);
+				Rect.H = FMath::Min(TextureResolution, Rect.H);
 
-			if( bRectPack )
-			{
-				if( LayoutRaster.Find( Rect ) )
+				const bool bRectPack = false;
+
+				if (bRectPack)
 				{
-					// Is best?
-					if( Rect.X + Rect.Y * TextureResolution < BestRect.X + BestRect.Y * TextureResolution )
+					if (LayoutRaster.Find(Rect))
 					{
-						BestOrientation = Orientation;
-						BestRect = Rect;
+						// Is best?
+						if (Rect.X + Rect.Y * TextureResolution < BestRect.X + BestRect.Y * TextureResolution)
+						{
+							BestOrientation = Orientation;
+							BestRect = Rect;
+						}
+					}
+					else
+					{
+						continue;
 					}
 				}
 				else
 				{
-					continue;
-				}
-			}
-			else
-			{
-				if ( LayoutVersion >= FMeshDescriptionOperations::ELightmapUVVersion::Segments && Orientation % 4 == 1 )
-				{
-					ChartRaster.FlipX( Rect );
-				}
-				else if ( LayoutVersion >= FMeshDescriptionOperations::ELightmapUVVersion::Segments && Orientation % 4 == 3 )
-				{
-					ChartRaster.FlipY( Rect );
-				}
-				else
-				{
-					int32 BeginRasterize = FPlatformTime::Cycles();
-					RasterizeChart( Chart, Rect.W, Rect.H );
-					RasterizeCycles += FPlatformTime::Cycles() - BeginRasterize;
-				}
+					if (LayoutVersion >= FMeshDescriptionOperations::ELightmapUVVersion::Segments && Orientation % 4 == 1)
+					{
+						ChartRaster.FlipX(Rect);
+					}
+					else if (LayoutVersion >= FMeshDescriptionOperations::ELightmapUVVersion::Segments && Orientation % 4 == 3)
+					{
+						ChartRaster.FlipY(Rect);
+					}
+					else
+					{
+						int32 BeginRasterize = FPlatformTime::Cycles();
+						RasterizeChart(Chart, Rect.W, Rect.H);
+						RasterizeCycles += FPlatformTime::Cycles() - BeginRasterize;
+					}
 
-				bool bFound = false;
+					bool bFound = false;
 
-				uint32 BeginFind = FPlatformTime::Cycles();
-				if ( LayoutVersion == FMeshDescriptionOperations::ELightmapUVVersion::BitByBit )
-				{
-					bFound = LayoutRaster.FindBitByBit( Rect, ChartRaster );
-				}
-				else if ( LayoutVersion >= FMeshDescriptionOperations::ELightmapUVVersion::Segments )
-				{
-					bFound = LayoutRaster.FindWithSegments( Rect, BestRect, ChartRaster );
-				}
-				FindCycles += FPlatformTime::Cycles() - BeginFind;
+					uint32 BeginFind = FPlatformTime::Cycles();
+					if (LayoutVersion == FMeshDescriptionOperations::ELightmapUVVersion::BitByBit)
+					{
+						bFound = LayoutRaster.FindBitByBit(Rect, ChartRaster);
+					}
+					else if (LayoutVersion >= FMeshDescriptionOperations::ELightmapUVVersion::Segments)
+					{
+						bFound = LayoutRaster.FindWithSegments(Rect, BestRect, ChartRaster);
+					}
+					FindCycles += FPlatformTime::Cycles() - BeginFind;
 
-				if( bFound )
-				{
-					// Is best?
-					if( Rect.X + Rect.Y * TextureResolution < BestRect.X + BestRect.Y * TextureResolution )
+					if (bFound)
 					{
-						BestChartRaster = ChartRaster;
+						// Is best?
+						if (Rect.X + Rect.Y * TextureResolution < BestRect.X + BestRect.Y * TextureResolution)
+						{
+							BestChartRaster = ChartRaster;
 
-						BestOrientation = Orientation;
-						BestRect = Rect;
+							BestOrientation = Orientation;
+							BestRect = Rect;
 
-						if ( BestRect.X == 0 && BestRect.Y == 0 )
-						{
-							// BestRect can't be beat, stop here
-							break;
+							if (BestRect.X == 0 && BestRect.Y == 0)
+							{
+								// BestRect can't be beat, stop here
+								break;
+							}
 						}
 					}
-				}
-				else
-				{
-					continue;
+					else
+					{
+						continue;
+					}
 				}
 			}
-		}
 
-		if( BestOrientation >= 0 )
-		{
-			// Add chart to layout
-			OrientChart( Chart, BestOrientation );
+			if (BestOrientation >= 0)
+			{
+				// Add chart to layout
+				OrientChart(Chart, BestOrientation);
 
-			LayoutRaster.Alloc( BestRect, BestChartRaster );
+				LayoutRaster.Alloc(BestRect, BestChartRaster);
 
-			Chart.PackingBias.X += BestRect.X;
-			Chart.PackingBias.Y += BestRect.Y;
-		}
-		else
-		{
-			// Found no orientation that fit
-			return false;
+				Chart.PackingBias.X += BestRect.X;
+				Chart.PackingBias.Y += BestRect.Y;
+			}
+			else
+			{
+				// Found no orientation that fit
+				return false;
+			}
 		}
-	}
-
-	double EndPackCharts = FPlatformTime::Seconds();
 
-	UE_LOG(LogMeshDescriptionLayoutUV, Display, TEXT("PackCharts: %s"),	*FPlatformTime::PrettyTime(EndPackCharts - BeginPackCharts));
-	UE_LOG(LogMeshDescriptionLayoutUV, Display, TEXT("  Rasterize: %u"),	RasterizeCycles);
-	UE_LOG(LogMeshDescriptionLayoutUV, Display, TEXT("  Find: %u"),		FindCycles);
+		double EndPackCharts = FPlatformTime::Seconds();
 
-	return true;
-}
+		UE_LOG(LogMeshDescriptionLayoutUV, Display, TEXT("PackCharts: %s"), *FPlatformTime::PrettyTime(EndPackCharts - BeginPackCharts));
+		UE_LOG(LogMeshDescriptionLayoutUV, Display, TEXT("  Rasterize: %u"), RasterizeCycles);
+		UE_LOG(LogMeshDescriptionLayoutUV, Display, TEXT("  Find: %u"), FindCycles);
 
-void FLayoutUV::OrientChart( FMeshChart& Chart, int32 Orientation )
-{
-	switch( Orientation )
-	{
-	case 0:
-		// 0 degrees
-		Chart.PackingScaleU = FVector2D( Chart.UVScale.X, 0 );
-		Chart.PackingScaleV = FVector2D( 0, Chart.UVScale.Y );
-		Chart.PackingBias = -Chart.MinUV.X * Chart.PackingScaleU - Chart.MinUV.Y * Chart.PackingScaleV + 0.5f;
-		break;
-	case 1:
-		// 0 degrees, flip x
-		Chart.PackingScaleU = FVector2D( -Chart.UVScale.X, 0 );
-		Chart.PackingScaleV = FVector2D( 0, Chart.UVScale.Y );
-		Chart.PackingBias = -Chart.MaxUV.X * Chart.PackingScaleU - Chart.MinUV.Y * Chart.PackingScaleV + 0.5f;
-		break;
-	case 2:
-		// 90 degrees
-		Chart.PackingScaleU = FVector2D( 0, -Chart.UVScale.X );
-		Chart.PackingScaleV = FVector2D( Chart.UVScale.Y, 0 );
-		Chart.PackingBias = -Chart.MaxUV.X * Chart.PackingScaleU - Chart.MinUV.Y * Chart.PackingScaleV + 0.5f;
-		break;
-	case 3:
-		// 90 degrees, flip x
-		Chart.PackingScaleU = FVector2D( 0, Chart.UVScale.X );
-		Chart.PackingScaleV = FVector2D( Chart.UVScale.Y, 0 );
-		Chart.PackingBias = -Chart.MinUV.X * Chart.PackingScaleU - Chart.MinUV.Y * Chart.PackingScaleV + 0.5f;
-		break;
-	case 4:
-		// 180 degrees
-		Chart.PackingScaleU = FVector2D( -Chart.UVScale.X, 0 );
-		Chart.PackingScaleV = FVector2D( 0, -Chart.UVScale.Y );
-		Chart.PackingBias = -Chart.MaxUV.X * Chart.PackingScaleU - Chart.MaxUV.Y * Chart.PackingScaleV + 0.5f;
-		break;
-	case 5:
-		// 180 degrees, flip x
-		Chart.PackingScaleU = FVector2D( Chart.UVScale.X, 0 );
-		Chart.PackingScaleV = FVector2D( 0, -Chart.UVScale.Y );
-		Chart.PackingBias = -Chart.MinUV.X * Chart.PackingScaleU - Chart.MaxUV.Y * Chart.PackingScaleV + 0.5f;
-		break;
-	case 6:
-		// 270 degrees
-		Chart.PackingScaleU = FVector2D( 0, Chart.UVScale.X );
-		Chart.PackingScaleV = FVector2D( -Chart.UVScale.Y, 0 );
-		Chart.PackingBias = -Chart.MinUV.X * Chart.PackingScaleU - Chart.MaxUV.Y * Chart.PackingScaleV + 0.5f;
-		break;
-	case 7:
-		// 270 degrees, flip x
-		Chart.PackingScaleU = FVector2D( 0, -Chart.UVScale.X );
-		Chart.PackingScaleV = FVector2D( -Chart.UVScale.Y, 0 );
-		Chart.PackingBias = -Chart.MaxUV.X * Chart.PackingScaleU - Chart.MaxUV.Y * Chart.PackingScaleV + 0.5f;
-		break;
+		return true;
 	}
-}
 
-// Max of 2048x2048 due to precision
-// Dilate in 28.4 fixed point. Half pixel dilation is conservative rasterization.
-// Dilation same as Minkowski sum of triangle and square.
-template< typename TShader, int32 Dilate >
-void RasterizeTriangle( TShader& Shader, const FVector2D Points[3], int32 ScissorWidth, int32 ScissorHeight )
-{
-	const FVector2D HalfPixel( 0.5f, 0.5f );
-	FVector2D p0 = Points[0] - HalfPixel;
-	FVector2D p1 = Points[1] - HalfPixel;
-	FVector2D p2 = Points[2] - HalfPixel;
-
-	// Correct winding
-	float Facing = ( p0.X - p1.X ) * ( p2.Y - p0.Y ) - ( p0.Y - p1.Y ) * ( p2.X - p0.X );
-	if( Facing < 0.0f )
+	void FLayoutUV::OrientChart(FMeshChart& Chart, int32 Orientation)
 	{
-		Swap( p0, p2 );
+		switch (Orientation)
+		{
+		case 0:
+			// 0 degrees
+			Chart.PackingScaleU = FVector2D(Chart.UVScale.X, 0);
+			Chart.PackingScaleV = FVector2D(0, Chart.UVScale.Y);
+			Chart.PackingBias = -Chart.MinUV.X * Chart.PackingScaleU - Chart.MinUV.Y * Chart.PackingScaleV + 0.5f;
+			break;
+		case 1:
+			// 0 degrees, flip x
+			Chart.PackingScaleU = FVector2D(-Chart.UVScale.X, 0);
+			Chart.PackingScaleV = FVector2D(0, Chart.UVScale.Y);
+			Chart.PackingBias = -Chart.MaxUV.X * Chart.PackingScaleU - Chart.MinUV.Y * Chart.PackingScaleV + 0.5f;
+			break;
+		case 2:
+			// 90 degrees
+			Chart.PackingScaleU = FVector2D(0, -Chart.UVScale.X);
+			Chart.PackingScaleV = FVector2D(Chart.UVScale.Y, 0);
+			Chart.PackingBias = -Chart.MaxUV.X * Chart.PackingScaleU - Chart.MinUV.Y * Chart.PackingScaleV + 0.5f;
+			break;
+		case 3:
+			// 90 degrees, flip x
+			Chart.PackingScaleU = FVector2D(0, Chart.UVScale.X);
+			Chart.PackingScaleV = FVector2D(Chart.UVScale.Y, 0);
+			Chart.PackingBias = -Chart.MinUV.X * Chart.PackingScaleU - Chart.MinUV.Y * Chart.PackingScaleV + 0.5f;
+			break;
+		case 4:
+			// 180 degrees
+			Chart.PackingScaleU = FVector2D(-Chart.UVScale.X, 0);
+			Chart.PackingScaleV = FVector2D(0, -Chart.UVScale.Y);
+			Chart.PackingBias = -Chart.MaxUV.X * Chart.PackingScaleU - Chart.MaxUV.Y * Chart.PackingScaleV + 0.5f;
+			break;
+		case 5:
+			// 180 degrees, flip x
+			Chart.PackingScaleU = FVector2D(Chart.UVScale.X, 0);
+			Chart.PackingScaleV = FVector2D(0, -Chart.UVScale.Y);
+			Chart.PackingBias = -Chart.MinUV.X * Chart.PackingScaleU - Chart.MaxUV.Y * Chart.PackingScaleV + 0.5f;
+			break;
+		case 6:
+			// 270 degrees
+			Chart.PackingScaleU = FVector2D(0, Chart.UVScale.X);
+			Chart.PackingScaleV = FVector2D(-Chart.UVScale.Y, 0);
+			Chart.PackingBias = -Chart.MinUV.X * Chart.PackingScaleU - Chart.MaxUV.Y * Chart.PackingScaleV + 0.5f;
+			break;
+		case 7:
+			// 270 degrees, flip x
+			Chart.PackingScaleU = FVector2D(0, -Chart.UVScale.X);
+			Chart.PackingScaleV = FVector2D(-Chart.UVScale.Y, 0);
+			Chart.PackingBias = -Chart.MaxUV.X * Chart.PackingScaleU - Chart.MaxUV.Y * Chart.PackingScaleV + 0.5f;
+			break;
+		}
 	}
 
-	// 28.4 fixed point
-	const int32 X0 = (int32)( 16.0f * p0.X + 0.5f );
-	const int32 X1 = (int32)( 16.0f * p1.X + 0.5f );
-	const int32 X2 = (int32)( 16.0f * p2.X + 0.5f );
-	
-	const int32 Y0 = (int32)( 16.0f * p0.Y + 0.5f );
-	const int32 Y1 = (int32)( 16.0f * p1.Y + 0.5f );
-	const int32 Y2 = (int32)( 16.0f * p2.Y + 0.5f );
-
-	// Bounding rect
-	int32 MinX = ( FMath::Min3( X0, X1, X2 ) - Dilate + 15 ) / 16;
-	int32 MaxX = ( FMath::Max3( X0, X1, X2 ) + Dilate + 15 ) / 16;
-	int32 MinY = ( FMath::Min3( Y0, Y1, Y2 ) - Dilate + 15 ) / 16;
-	int32 MaxY = ( FMath::Max3( Y0, Y1, Y2 ) + Dilate + 15 ) / 16;
-
-	// Clip to image
-	MinX = FMath::Clamp( MinX, 0, ScissorWidth );
-	MaxX = FMath::Clamp( MaxX, 0, ScissorWidth );
-	MinY = FMath::Clamp( MinY, 0, ScissorHeight );
-	MaxY = FMath::Clamp( MaxY, 0, ScissorHeight );
-
-	// Deltas
-	const int32 DX01 = X0 - X1;
-	const int32 DX12 = X1 - X2;
-	const int32 DX20 = X2 - X0;
-
-	const int32 DY01 = Y0 - Y1;
-	const int32 DY12 = Y1 - Y2;
-	const int32 DY20 = Y2 - Y0;
-
-	// Half-edge constants
-	int32 C0 = DY01 * X0 - DX01 * Y0;
-	int32 C1 = DY12 * X1 - DX12 * Y1;
-	int32 C2 = DY20 * X2 - DX20 * Y2;
-
-	// Correct for fill convention
-	C0 += ( DY01 < 0 || ( DY01 == 0 && DX01 > 0 ) ) ? 0 : -1;
-	C1 += ( DY12 < 0 || ( DY12 == 0 && DX12 > 0 ) ) ? 0 : -1;
-	C2 += ( DY20 < 0 || ( DY20 == 0 && DX20 > 0 ) ) ? 0 : -1;
-
-	// Dilate edges
-	C0 += ( abs(DX01) + abs(DY01) ) * Dilate;
-	C1 += ( abs(DX12) + abs(DY12) ) * Dilate;
-	C2 += ( abs(DX20) + abs(DY20) ) * Dilate;
-
-	for( int32 y = MinY; y < MaxY; y++ )
+	// Max of 2048x2048 due to precision
+	// Dilate in 28.4 fixed point. Half pixel dilation is conservative rasterization.
+	// Dilation same as Minkowski sum of triangle and square.
+	template< typename TShader, int32 Dilate >
+	void RasterizeTriangle(TShader& Shader, const FVector2D Points[3], int32 ScissorWidth, int32 ScissorHeight)
 	{
-		for( int32 x = MinX; x < MaxX; x++ )
+		const FVector2D HalfPixel(0.5f, 0.5f);
+		FVector2D p0 = Points[0] - HalfPixel;
+		FVector2D p1 = Points[1] - HalfPixel;
+		FVector2D p2 = Points[2] - HalfPixel;
+
+		// Correct winding
+		float Facing = (p0.X - p1.X) * (p2.Y - p0.Y) - (p0.Y - p1.Y) * (p2.X - p0.X);
+		if (Facing < 0.0f)
 		{
-			// same as Edge1 >= 0 && Edge2 >= 0 && Edge3 >= 0
-			int32 IsInside;
-			IsInside  = C0 + (DX01 * y - DY01 * x) * 16;
-			IsInside |= C1 + (DX12 * y - DY12 * x) * 16;
-			IsInside |= C2 + (DX20 * y - DY20 * x) * 16;
+			Swap(p0, p2);
+		}
 
-			if( IsInside >= 0 )
+		// 28.4 fixed point
+		const int32 X0 = (int32)(16.0f * p0.X + 0.5f);
+		const int32 X1 = (int32)(16.0f * p1.X + 0.5f);
+		const int32 X2 = (int32)(16.0f * p2.X + 0.5f);
+
+		const int32 Y0 = (int32)(16.0f * p0.Y + 0.5f);
+		const int32 Y1 = (int32)(16.0f * p1.Y + 0.5f);
+		const int32 Y2 = (int32)(16.0f * p2.Y + 0.5f);
+
+		// Bounding rect
+		int32 MinX = (FMath::Min3(X0, X1, X2) - Dilate + 15) / 16;
+		int32 MaxX = (FMath::Max3(X0, X1, X2) + Dilate + 15) / 16;
+		int32 MinY = (FMath::Min3(Y0, Y1, Y2) - Dilate + 15) / 16;
+		int32 MaxY = (FMath::Max3(Y0, Y1, Y2) + Dilate + 15) / 16;
+
+		// Clip to image
+		MinX = FMath::Clamp(MinX, 0, ScissorWidth);
+		MaxX = FMath::Clamp(MaxX, 0, ScissorWidth);
+		MinY = FMath::Clamp(MinY, 0, ScissorHeight);
+		MaxY = FMath::Clamp(MaxY, 0, ScissorHeight);
+
+		// Deltas
+		const int32 DX01 = X0 - X1;
+		const int32 DX12 = X1 - X2;
+		const int32 DX20 = X2 - X0;
+
+		const int32 DY01 = Y0 - Y1;
+		const int32 DY12 = Y1 - Y2;
+		const int32 DY20 = Y2 - Y0;
+
+		// Half-edge constants
+		int32 C0 = DY01 * X0 - DX01 * Y0;
+		int32 C1 = DY12 * X1 - DX12 * Y1;
+		int32 C2 = DY20 * X2 - DX20 * Y2;
+
+		// Correct for fill convention
+		C0 += (DY01 < 0 || (DY01 == 0 && DX01 > 0)) ? 0 : -1;
+		C1 += (DY12 < 0 || (DY12 == 0 && DX12 > 0)) ? 0 : -1;
+		C2 += (DY20 < 0 || (DY20 == 0 && DX20 > 0)) ? 0 : -1;
+
+		// Dilate edges
+		C0 += (abs(DX01) + abs(DY01)) * Dilate;
+		C1 += (abs(DX12) + abs(DY12)) * Dilate;
+		C2 += (abs(DX20) + abs(DY20)) * Dilate;
+
+		for (int32 y = MinY; y < MaxY; y++)
+		{
+			for (int32 x = MinX; x < MaxX; x++)
 			{
-				Shader.Process( x, y );
+				// same as Edge1 >= 0 && Edge2 >= 0 && Edge3 >= 0
+				int32 IsInside;
+				IsInside = C0 + (DX01 * y - DY01 * x) * 16;
+				IsInside |= C1 + (DX12 * y - DY12 * x) * 16;
+				IsInside |= C2 + (DX20 * y - DY20 * x) * 16;
+
+				if (IsInside >= 0)
+				{
+					Shader.Process(x, y);
+				}
 			}
 		}
 	}
-}
 
-void FLayoutUV::RasterizeChart( const FMeshChart& Chart, uint32 RectW, uint32 RectH )
-{
-	// Bilinear footprint is -1 to 1 pixels. If packed geometrically, only a half pixel dilation
-	// would be needed to guarantee all charts were at least 1 pixel away, safe for bilinear filtering.
-	// Unfortunately, with pixel packing a full 1 pixel dilation is required unless chart edges exactly
-	// align with pixel centers.
+	void FLayoutUV::RasterizeChart(const FMeshChart& Chart, uint32 RectW, uint32 RectH)
+	{
+		// Bilinear footprint is -1 to 1 pixels. If packed geometrically, only a half pixel dilation
+		// would be needed to guarantee all charts were at least 1 pixel away, safe for bilinear filtering.
+		// Unfortunately, with pixel packing a full 1 pixel dilation is required unless chart edges exactly
+		// align with pixel centers.
 
-	ChartRaster.Clear();
+		ChartRaster.Clear();
 
-	for( uint32 Tri = Chart.FirstTri; Tri < Chart.LastTri; Tri++ )
-	{
-		FVector2D Points[3];
-		for( int k = 0; k < 3; k++ )
+		for (uint32 Tri = Chart.FirstTri; Tri < Chart.LastTri; Tri++)
 		{
-			const FVector2D& UV = TexCoords[ 3 * SortedTris[ Tri ] + k ];
-			Points[k] = UV.X * Chart.PackingScaleU + UV.Y * Chart.PackingScaleV + Chart.PackingBias;
-		}
+			FVector2D Points[3];
+			for (int k = 0; k < 3; k++)
+			{
+				const FVector2D& UV = TexCoords[3 * SortedTris[Tri] + k];
+				Points[k] = UV.X * Chart.PackingScaleU + UV.Y * Chart.PackingScaleV + Chart.PackingBias;
+			}
 
-		RasterizeTriangle< FAllocator2DShader, 16 >( ChartShader, Points, RectW, RectH );
-	}
+			RasterizeTriangle< FAllocator2DShader, 16 >(ChartShader, Points, RectW, RectH);
+		}
 
-	if ( LayoutVersion >= FMeshDescriptionOperations::ELightmapUVVersion::Segments )
-	{
-		ChartRaster.CreateUsedSegments();
+		if (LayoutVersion >= FMeshDescriptionOperations::ELightmapUVVersion::Segments)
+		{
+			ChartRaster.CreateUsedSegments();
+		}
 	}
-}
 
-void FLayoutUV::CommitPackedUVs()
-{
-	// If current DstChannel is out of range of the number of UVs defined by the mesh description, change the index count accordingly
-	const uint32 NumUVs = MeshDescription->VertexInstanceAttributes().GetAttributeIndexCount<FVector2D>(MeshAttribute::VertexInstance::TextureCoordinate);
-	if (DstChannel >= NumUVs)
+	void FLayoutUV::CommitPackedUVs()
 	{
-		MeshDescription->VertexInstanceAttributes().SetAttributeIndexCount<FVector2D>(MeshAttribute::VertexInstance::TextureCoordinate, DstChannel + 1);
-		ensure(false);	// not expecting it to get here
-	}
+		// If current DstChannel is out of range of the number of UVs defined by the mesh description, change the index count accordingly
+		const uint32 NumUVs = MeshDescription->VertexInstanceAttributes().GetAttributeIndexCount<FVector2D>(MeshAttribute::VertexInstance::TextureCoordinate);
+		if (DstChannel >= NumUVs)
+		{
+			MeshDescription->VertexInstanceAttributes().SetAttributeIndexCount<FVector2D>(MeshAttribute::VertexInstance::TextureCoordinate, DstChannel + 1);
+			ensure(false);	// not expecting it to get here
+		}
 
-	TVertexInstanceAttributeArray<FVector2D>& VertexUVs = MeshDescription->VertexInstanceAttributes().GetAttributes<FVector2D>(MeshAttribute::VertexInstance::TextureCoordinate, DstChannel);
+		TVertexInstanceAttributeArray<FVector2D>& VertexUVs = MeshDescription->VertexInstanceAttributes().GetAttributes<FVector2D>(MeshAttribute::VertexInstance::TextureCoordinate, DstChannel);
 
-	// Commit chart UVs
-	for( int32 i = 0; i < Charts.Num(); i++ )
-	{
-		FMeshChart& Chart = Charts[i];
+		// Commit chart UVs
+		for (int32 i = 0; i < Charts.Num(); i++)
+		{
+			FMeshChart& Chart = Charts[i];
 
-		Chart.PackingScaleU /= TextureResolution;
-		Chart.PackingScaleV /= TextureResolution;
-		Chart.PackingBias /= TextureResolution;
+			Chart.PackingScaleU /= TextureResolution;
+			Chart.PackingScaleV /= TextureResolution;
+			Chart.PackingBias /= TextureResolution;
 
-		for( uint32 Tri = Chart.FirstTri; Tri < Chart.LastTri; Tri++ )
-		{
-			for( int k = 0; k < 3; k++ )
+			for (uint32 Tri = Chart.FirstTri; Tri < Chart.LastTri; Tri++)
 			{
-				uint32 Index = 3 * SortedTris[ Tri ] + k;
-				const FVector2D& UV = TexCoords[ Index ];
-				const FVertexInstanceID VertexInstanceID(Index);
-				VertexUVs[VertexInstanceID] = UV.X * Chart.PackingScaleU + UV.Y * Chart.PackingScaleV + Chart.PackingBias;
+				for (int k = 0; k < 3; k++)
+				{
+					uint32 Index = 3 * SortedTris[Tri] + k;
+					const FVector2D& UV = TexCoords[Index];
+					const FVertexInstanceID VertexInstanceID(Index);
+					VertexUVs[VertexInstanceID] = UV.X * Chart.PackingScaleU + UV.Y * Chart.PackingScaleV + Chart.PackingBias;
+				}
 			}
 		}
 	}
-}
+}
\ No newline at end of file
diff --git a/Engine/Source/Developer/MeshDescriptionOperations/Private/LayoutUV.h b/Engine/Source/Developer/MeshDescriptionOperations/Private/LayoutUV.h
--- a/Engine/Source/Developer/MeshDescriptionOperations/Private/LayoutUV.h
+++ b/Engine/Source/Developer/MeshDescriptionOperations/Private/LayoutUV.h
@@ -10,158 +10,160 @@
 
 #define NEW_UVS_ARE_SAME THRESH_POINTS_ARE_SAME
 #define LEGACY_UVS_ARE_SAME (1.0f / 1024.0f)
-
-struct FMeshChart
-{
-	uint32		FirstTri;
-	uint32		LastTri;
-	
-	FVector2D	MinUV;
-	FVector2D	MaxUV;
-	
-	float		UVArea;
-	FVector2D	UVScale;
-	FVector2D	WorldScale;
-	
-	FVector2D	PackingScaleU;
-	FVector2D	PackingScaleV;
-	FVector2D	PackingBias;
-
-	int32		Join[4];
-};
-
-struct FAllocator2DShader
+namespace MeshDescriptionOp
 {
-	FAllocator2D*	Allocator2D;
+	struct FMeshChart
+	{
+		uint32		FirstTri;
+		uint32		LastTri;
 
-	FAllocator2DShader( FAllocator2D* InAllocator2D )
-	: Allocator2D( InAllocator2D )
-	{}
+		FVector2D	MinUV;
+		FVector2D	MaxUV;
 
-	FORCEINLINE void Process( uint32 x, uint32 y )
-	{
-		Allocator2D->SetBit( x, y );
-	}
-};
+		float		UVArea;
+		FVector2D	UVScale;
+		FVector2D	WorldScale;
 
-class FLayoutUV
-{
-public:
-				FLayoutUV( UMeshDescription* InMesh, uint32 InSrcChannel, uint32 InDstChannel, uint32 InTextureResolution );
-
-	void		FindCharts( const TMultiMap<int32,int32>& OverlappingCorners );
-	bool		FindBestPacking();
-	void		CommitPackedUVs();
-
-	void		SetVersion(FMeshDescriptionOperations::ELightmapUVVersion Version ) { LayoutVersion = Version; }
-
-private:
-	bool		PositionsMatch( uint32 a, uint32 b ) const;
-	bool		NormalsMatch( uint32 a, uint32 b ) const;
-	bool		UVsMatch( uint32 a, uint32 b ) const;
-	bool		VertsMatch( uint32 a, uint32 b ) const;
-	float		TriangleUVArea( uint32 Tri ) const;
-	void		DisconnectChart( FMeshChart& Chart, uint32 Side );
-
-	void		ScaleCharts( float UVScale );
-	bool		PackCharts();
-	void		OrientChart( FMeshChart& Chart, int32 Orientation );
-	void		RasterizeChart( const FMeshChart& Chart, uint32 RectW, uint32 RectH );
-
-	float		GetUVEqualityThreshold() const { return LayoutVersion >= FMeshDescriptionOperations::ELightmapUVVersion::SmallChartPacking ? NEW_UVS_ARE_SAME : LEGACY_UVS_ARE_SAME; }
-
-	UMeshDescription* MeshDescription;
-	uint32		SrcChannel;
-	uint32		DstChannel;
-	uint32		TextureResolution;
-	
-	TArray< FVector2D >		TexCoords;
-	TArray< uint32 >		SortedTris;
-	TArray< FMeshChart >	Charts;
-	float					TotalUVArea;
-	float					MaxChartSize;
-
-	FAllocator2D		LayoutRaster;
-	FAllocator2D		ChartRaster;
-	FAllocator2D		BestChartRaster;
-	FAllocator2DShader	ChartShader;
-
-	FMeshDescriptionOperations::ELightmapUVVersion LayoutVersion;
-};
-
-
-inline bool FLayoutUV::PositionsMatch( uint32 a, uint32 b ) const
-{
-	const FVertexInstanceID VertexInstanceIDA(a);
-	const FVertexInstanceID VertexInstanceIDB(b);
-	const FVertexID VertexIDA = MeshDescription->GetVertexInstanceVertex(VertexInstanceIDA);
-	const FVertexID VertexIDB = MeshDescription->GetVertexInstanceVertex(VertexInstanceIDB);
+		FVector2D	PackingScaleU;
+		FVector2D	PackingScaleV;
+		FVector2D	PackingBias;
 
-	const TVertexAttributeArray<FVector>& VertexPositions = MeshDescription->VertexAttributes().GetAttributes<FVector>(MeshAttribute::Vertex::Position);
-	return VertexPositions[VertexIDA].Equals(VertexPositions[VertexIDB], THRESH_POINTS_ARE_SAME);
-}
+		int32		Join[4];
+	};
 
-inline bool FLayoutUV::NormalsMatch( uint32 a, uint32 b ) const
-{
-	// If current SrcChannel is out of range of the number of UVs defined by the mesh description, just return true
-	// @todo: hopefully remove this check entirely and just ensure that the mesh description matches the inputs
-	const uint32 NumUVs = MeshDescription->VertexInstanceAttributes().GetAttributeIndexCount<FVector2D>(MeshAttribute::VertexInstance::TextureCoordinate);
-	if (SrcChannel >= NumUVs)
+	struct FAllocator2DShader
 	{
-		ensure(false);	// not expecting it to get here
-		return true;
-	}
+		FAllocator2D*	Allocator2D;
 
-	const FVertexInstanceID VertexInstanceIDA(a);
-	const FVertexInstanceID VertexInstanceIDB(b);
+		FAllocator2DShader(FAllocator2D* InAllocator2D)
+			: Allocator2D(InAllocator2D)
+		{}
 
-	const TVertexInstanceAttributeArray<FVector>& VertexNormals = MeshDescription->VertexInstanceAttributes().GetAttributes<FVector>(MeshAttribute::VertexInstance::Normal);
-	return VertexNormals[VertexInstanceIDA].Equals(VertexNormals[VertexInstanceIDB], THRESH_NORMALS_ARE_SAME);
-}
+		FORCEINLINE void Process(uint32 x, uint32 y)
+		{
+			Allocator2D->SetBit(x, y);
+		}
+	};
 
-inline bool FLayoutUV::UVsMatch( uint32 a, uint32 b ) const
-{
-	// If current SrcChannel is out of range of the number of UVs defined by the mesh description, just return true
-	const uint32 NumUVs = MeshDescription->VertexInstanceAttributes().GetAttributeIndexCount<FVector2D>(MeshAttribute::VertexInstance::TextureCoordinate);
-	if (SrcChannel >= NumUVs)
+	class FLayoutUV
 	{
-		ensure(false);	// not expecting it to get here
-		return true;
-	}
+	public:
+		FLayoutUV(UMeshDescription* InMesh, uint32 InSrcChannel, uint32 InDstChannel, uint32 InTextureResolution);
 
-	const FVertexInstanceID VertexInstanceIDA(a);
-	const FVertexInstanceID VertexInstanceIDB(b);
+		void		FindCharts(const TMultiMap<int32, int32>& OverlappingCorners);
+		bool		FindBestPacking();
+		void		CommitPackedUVs();
 
-	const TVertexInstanceAttributeArray<FVector2D>& VertexUVs = MeshDescription->VertexInstanceAttributes().GetAttributes<FVector2D>(MeshAttribute::VertexInstance::TextureCoordinate, SrcChannel);
-	return VertexUVs[VertexInstanceIDA].Equals(VertexUVs[VertexInstanceIDB], GetUVEqualityThreshold());
-}
+		void		SetVersion(FMeshDescriptionOperations::ELightmapUVVersion Version) { LayoutVersion = Version; }
 
-inline bool FLayoutUV::VertsMatch( uint32 a, uint32 b ) const
-{
-	return PositionsMatch(a, b) && UVsMatch(a, b);
-}
+	private:
+		bool		PositionsMatch(uint32 a, uint32 b) const;
+		bool		NormalsMatch(uint32 a, uint32 b) const;
+		bool		UVsMatch(uint32 a, uint32 b) const;
+		bool		VertsMatch(uint32 a, uint32 b) const;
+		float		TriangleUVArea(uint32 Tri) const;
+		void		DisconnectChart(FMeshChart& Chart, uint32 Side);
 
-// Signed UV area
-inline float FLayoutUV::TriangleUVArea( uint32 Tri ) const
-{
-	const TVertexInstanceAttributeArray<FVector2D>& VertexUVs = MeshDescription->VertexInstanceAttributes().GetAttributes<FVector2D>(MeshAttribute::VertexInstance::TextureCoordinate, SrcChannel);
+		void		ScaleCharts(float UVScale);
+		bool		PackCharts();
+		void		OrientChart(FMeshChart& Chart, int32 Orientation);
+		void		RasterizeChart(const FMeshChart& Chart, uint32 RectW, uint32 RectH);
+
+		float		GetUVEqualityThreshold() const { return LayoutVersion >= FMeshDescriptionOperations::ELightmapUVVersion::SmallChartPacking ? NEW_UVS_ARE_SAME : LEGACY_UVS_ARE_SAME; }
+
+		UMeshDescription* MeshDescription;
+		uint32		SrcChannel;
+		uint32		DstChannel;
+		uint32		TextureResolution;
+
+		TArray< FVector2D >		TexCoords;
+		TArray< uint32 >		SortedTris;
+		TArray< FMeshChart >	Charts;
+		float					TotalUVArea;
+		float					MaxChartSize;
+
+		FAllocator2D		LayoutRaster;
+		FAllocator2D		ChartRaster;
+		FAllocator2D		BestChartRaster;
+		FAllocator2DShader	ChartShader;
 
-	FVector2D UVs[3];
-	for( int k = 0; k < 3; k++ )
+		FMeshDescriptionOperations::ELightmapUVVersion LayoutVersion;
+	};
+
+
+	inline bool FLayoutUV::PositionsMatch(uint32 a, uint32 b) const
 	{
-		UVs[k] = VertexUVs[FVertexInstanceID((3*Tri) + k)];
+		const FVertexInstanceID VertexInstanceIDA(a);
+		const FVertexInstanceID VertexInstanceIDB(b);
+		const FVertexID VertexIDA = MeshDescription->GetVertexInstanceVertex(VertexInstanceIDA);
+		const FVertexID VertexIDB = MeshDescription->GetVertexInstanceVertex(VertexInstanceIDB);
+
+		const TVertexAttributeArray<FVector>& VertexPositions = MeshDescription->VertexAttributes().GetAttributes<FVector>(MeshAttribute::Vertex::Position);
+		return VertexPositions[VertexIDA].Equals(VertexPositions[VertexIDB], THRESH_POINTS_ARE_SAME);
 	}
 
-	FVector2D EdgeUV1 = UVs[1] - UVs[0];
-	FVector2D EdgeUV2 = UVs[2] - UVs[0];
-	return 0.5f * ( EdgeUV1.X * EdgeUV2.Y - EdgeUV1.Y * EdgeUV2.X );
-}
+	inline bool FLayoutUV::NormalsMatch(uint32 a, uint32 b) const
+	{
+		// If current SrcChannel is out of range of the number of UVs defined by the mesh description, just return true
+		// @todo: hopefully remove this check entirely and just ensure that the mesh description matches the inputs
+		const uint32 NumUVs = MeshDescription->VertexInstanceAttributes().GetAttributeIndexCount<FVector2D>(MeshAttribute::VertexInstance::TextureCoordinate);
+		if (SrcChannel >= NumUVs)
+		{
+			ensure(false);	// not expecting it to get here
+			return true;
+		}
+
+		const FVertexInstanceID VertexInstanceIDA(a);
+		const FVertexInstanceID VertexInstanceIDB(b);
+
+		const TVertexInstanceAttributeArray<FVector>& VertexNormals = MeshDescription->VertexInstanceAttributes().GetAttributes<FVector>(MeshAttribute::VertexInstance::Normal);
+		return VertexNormals[VertexInstanceIDA].Equals(VertexNormals[VertexInstanceIDB], THRESH_NORMALS_ARE_SAME);
+	}
 
-inline void FLayoutUV::DisconnectChart( FMeshChart& Chart, uint32 Side )
-{
-	if( Chart.Join[ Side ] != -1 )
+	inline bool FLayoutUV::UVsMatch(uint32 a, uint32 b) const
+	{
+		// If current SrcChannel is out of range of the number of UVs defined by the mesh description, just return true
+		const uint32 NumUVs = MeshDescription->VertexInstanceAttributes().GetAttributeIndexCount<FVector2D>(MeshAttribute::VertexInstance::TextureCoordinate);
+		if (SrcChannel >= NumUVs)
+		{
+			ensure(false);	// not expecting it to get here
+			return true;
+		}
+
+		const FVertexInstanceID VertexInstanceIDA(a);
+		const FVertexInstanceID VertexInstanceIDB(b);
+
+		const TVertexInstanceAttributeArray<FVector2D>& VertexUVs = MeshDescription->VertexInstanceAttributes().GetAttributes<FVector2D>(MeshAttribute::VertexInstance::TextureCoordinate, SrcChannel);
+		return VertexUVs[VertexInstanceIDA].Equals(VertexUVs[VertexInstanceIDB], GetUVEqualityThreshold());
+	}
+
+	inline bool FLayoutUV::VertsMatch(uint32 a, uint32 b) const
+	{
+		return PositionsMatch(a, b) && UVsMatch(a, b);
+	}
+
+	// Signed UV area
+	inline float FLayoutUV::TriangleUVArea(uint32 Tri) const
+	{
+		const TVertexInstanceAttributeArray<FVector2D>& VertexUVs = MeshDescription->VertexInstanceAttributes().GetAttributes<FVector2D>(MeshAttribute::VertexInstance::TextureCoordinate, SrcChannel);
+
+		FVector2D UVs[3];
+		for (int k = 0; k < 3; k++)
+		{
+			UVs[k] = VertexUVs[FVertexInstanceID((3 * Tri) + k)];
+		}
+
+		FVector2D EdgeUV1 = UVs[1] - UVs[0];
+		FVector2D EdgeUV2 = UVs[2] - UVs[0];
+		return 0.5f * (EdgeUV1.X * EdgeUV2.Y - EdgeUV1.Y * EdgeUV2.X);
+	}
+
+	inline void FLayoutUV::DisconnectChart(FMeshChart& Chart, uint32 Side)
 	{
-		Charts[ Chart.Join[ Side ] ].Join[ Side ^ 1 ] = -1;
-		Chart.Join[ Side ] = -1;
+		if (Chart.Join[Side] != -1)
+		{
+			Charts[Chart.Join[Side]].Join[Side ^ 1] = -1;
+			Chart.Join[Side] = -1;
+		}
 	}
-}
+}
\ No newline at end of file
diff --git a/Engine/Source/Developer/MeshDescriptionOperations/Private/MeshDescriptionOperations.cpp b/Engine/Source/Developer/MeshDescriptionOperations/Private/MeshDescriptionOperations.cpp
--- a/Engine/Source/Developer/MeshDescriptionOperations/Private/MeshDescriptionOperations.cpp
+++ b/Engine/Source/Developer/MeshDescriptionOperations/Private/MeshDescriptionOperations.cpp
@@ -353,7 +353,7 @@ void FillMeshDescriptionVertexPositionNoDuplicate(const TArray<FVector> &RawMesh
 		}
 		TempRemapVertexPosition.FindOrAdd(Index_i) = VertexCount;
 		// only need to search forward, since we add pairs both ways
-/*		for (int32 j = i + 1; j < VertIndexAndZ.Num(); j++)
+		for (int32 j = i + 1; j < VertIndexAndZ.Num(); j++)
 		{
 			if (FMath::Abs(VertIndexAndZ[j].Z - VertIndexAndZ[i].Z) > SMALL_NUMBER)
 				break; // can't be any more dups
@@ -365,7 +365,7 @@ void FillMeshDescriptionVertexPositionNoDuplicate(const TArray<FVector> &RawMesh
 			{
 				TempRemapVertexPosition.FindOrAdd(VertIndexAndZ[j].Index) = VertexCount;
 			}
-		}*/
+		}
 		VertexCount++;
 	}
 
@@ -1053,7 +1053,7 @@ void FMeshDescriptionOperations::CreateLightMapUVLayout(UMeshDescription* MeshDe
 	ELightmapUVVersion LightmapUVVersion,
 	const TMultiMap<int32, int32>& OverlappingCorners)
 {
-	FLayoutUV Packer(MeshDescription, SrcLightmapIndex, DstLightmapIndex, MinLightmapResolution);
+	MeshDescriptionOp::FLayoutUV Packer(MeshDescription, SrcLightmapIndex, DstLightmapIndex, MinLightmapResolution);
 	Packer.SetVersion(LightmapUVVersion);
 
 	Packer.FindCharts(OverlappingCorners);
@@ -1074,7 +1074,7 @@ bool FMeshDescriptionOperations::GenerateUniqueUVsForStaticMesh(const UMeshDescr
 	FindOverlappingCorners(OverlappingCorners, DuplicateMeshDescription, THRESH_POINTS_ARE_SAME);
 
 	// Generate new UVs
-	FLayoutUV Packer(DuplicateMeshDescription, 0, 1, FMath::Clamp(TextureResolution / 4, 32, 512));
+	MeshDescriptionOp::FLayoutUV Packer(DuplicateMeshDescription, 0, 1, FMath::Clamp(TextureResolution / 4, 32, 512));
 	Packer.FindCharts(OverlappingCorners);
 
 	bool bPackSuccess = Packer.FindBestPacking();
diff --git a/Engine/Source/Developer/MeshUtilities/MeshUtilities.Build.cs b/Engine/Source/Developer/MeshUtilities/MeshUtilities.Build.cs
--- a/Engine/Source/Developer/MeshUtilities/MeshUtilities.Build.cs
+++ b/Engine/Source/Developer/MeshUtilities/MeshUtilities.Build.cs
@@ -37,6 +37,7 @@ public class MeshUtilities : ModuleRules
 				"PropertyEditor",
 				"EditorStyle",
                 "GraphColor",
+                "MeshBuilder",
             }
 		);
 
diff --git a/Engine/Source/Developer/MeshUtilities/Private/MeshUtilities.cpp b/Engine/Source/Developer/MeshUtilities/Private/MeshUtilities.cpp
--- a/Engine/Source/Developer/MeshUtilities/Private/MeshUtilities.cpp
+++ b/Engine/Source/Developer/MeshUtilities/Private/MeshUtilities.cpp
@@ -26,9 +26,7 @@
 #include "RawMesh.h"
 #include "StaticMeshResources.h"
 #include "MeshBuild.h"
-#include "NvTriStrip.h"
-#include "forsythtriangleorderoptimizer.h"
-#include "nvtess.h"
+#include "ThirdPartyBuildOptimizationHelper.h"
 #include "SkeletalMeshTools.h"
 #include "Engine/SkeletalMesh.h"
 #include "Components/SkinnedMeshComponent.h"
@@ -118,338 +116,16 @@ MeshUtilities module.
 
 #define LOCTEXT_NAMESPACE "MeshUtils"
 
-// CVars
-static TAutoConsoleVariable<int32> CVarTriangleOrderOptimization(
-	TEXT("r.TriangleOrderOptimization"),
-	1,
-	TEXT("Controls the algorithm to use when optimizing the triangle order for the post-transform cache.\n")
-	TEXT("0: Use NVTriStrip (slower)\n")
-	TEXT("1: Use Forsyth algorithm (fastest)(default)")
-	TEXT("2: No triangle order optimization. (least efficient, debugging purposes only)"),
-	ECVF_Default);
-
 IMPLEMENT_MODULE(FMeshUtilities, MeshUtilities);
 
-/*------------------------------------------------------------------------------
-NVTriStrip for cache optimizing index buffers.
-------------------------------------------------------------------------------*/
-
-namespace NvTriStrip
-{
-	/**
-	* Converts 16 bit indices to 32 bit prior to passing them into the real GenerateStrips util method
-	*/
-	void GenerateStrips(
-		const uint8* Indices,
-		bool Is32Bit,
-		const uint32 NumIndices,
-		PrimitiveGroup** PrimGroups,
-		uint32* NumGroups
-		)
-	{
-		if (Is32Bit)
-		{
-			GenerateStrips((uint32*)Indices, NumIndices, PrimGroups, NumGroups);
-		}
-		else
-		{
-			// convert to 32 bit
-			uint32 Idx;
-			TArray<uint32> NewIndices;
-			NewIndices.AddUninitialized(NumIndices);
-			for (Idx = 0; Idx < NumIndices; ++Idx)
-			{
-				NewIndices[Idx] = ((uint16*)Indices)[Idx];
-			}
-			GenerateStrips(NewIndices.GetData(), NumIndices, PrimGroups, NumGroups);
-		}
-
-	}
-
-	/**
-	* Orders a triangle list for better vertex cache coherency.
-	*
-	* *** WARNING: This is safe to call for multiple threads IF AND ONLY IF all
-	* threads call SetListsOnly(true) and SetCacheSize(CACHESIZE_GEFORCE3). If
-	* NvTriStrip is ever used with different settings the library will need
-	* some modifications to be thread-safe. ***
-	*/
-	template<typename IndexDataType, typename Allocator>
-	void CacheOptimizeIndexBuffer(TArray<IndexDataType, Allocator>& Indices)
-	{
-		static_assert(sizeof(IndexDataType) == 2 || sizeof(IndexDataType) == 4, "Indices must be short or int.");
-
-		PrimitiveGroup*	PrimitiveGroups = NULL;
-		uint32			NumPrimitiveGroups = 0;
-		bool Is32Bit = sizeof(IndexDataType) == 4;
-
-		SetListsOnly(true);
-		SetCacheSize(CACHESIZE_GEFORCE3);
-
-		GenerateStrips((uint8*)Indices.GetData(), Is32Bit, Indices.Num(), &PrimitiveGroups, &NumPrimitiveGroups);
-
-		Indices.Empty();
-		Indices.AddUninitialized(PrimitiveGroups->numIndices);
-
-		if (Is32Bit)
-		{
-			FMemory::Memcpy(Indices.GetData(), PrimitiveGroups->indices, Indices.Num() * sizeof(IndexDataType));
-		}
-		else
-		{
-			for (uint32 I = 0; I < PrimitiveGroups->numIndices; ++I)
-			{
-				Indices[I] = (uint16)PrimitiveGroups->indices[I];
-			}
-		}
-
-		delete[] PrimitiveGroups;
-	}
-}
-
-/*------------------------------------------------------------------------------
-Forsyth algorithm for cache optimizing index buffers.
-------------------------------------------------------------------------------*/
-
-namespace Forsyth
-{
-	/**
-	* Converts 16 bit indices to 32 bit prior to passing them into the real OptimizeFaces util method
-	*/
-	void OptimizeFaces(
-		const uint8* Indices,
-		bool Is32Bit,
-		const uint32 NumIndices,
-		uint32 NumVertices,
-		uint32* OutIndices,
-		uint16 CacheSize
-		)
-	{
-		if (Is32Bit)
-		{
-			OptimizeFaces((uint32*)Indices, NumIndices, NumVertices, OutIndices, CacheSize);
-		}
-		else
-		{
-			// convert to 32 bit
-			uint32 Idx;
-			TArray<uint32> NewIndices;
-			NewIndices.AddUninitialized(NumIndices);
-			for (Idx = 0; Idx < NumIndices; ++Idx)
-			{
-				NewIndices[Idx] = ((uint16*)Indices)[Idx];
-			}
-			OptimizeFaces(NewIndices.GetData(), NumIndices, NumVertices, OutIndices, CacheSize);
-		}
-
-	}
-
-	/**
-	* Orders a triangle list for better vertex cache coherency.
-	*/
-	template<typename IndexDataType, typename Allocator>
-	void CacheOptimizeIndexBuffer(TArray<IndexDataType, Allocator>& Indices)
-	{
-		static_assert(sizeof(IndexDataType) == 2 || sizeof(IndexDataType) == 4, "Indices must be short or int.");
-		bool Is32Bit = sizeof(IndexDataType) == 4;
-
-		// Count the number of vertices
-		uint32 NumVertices = 0;
-		for (int32 Index = 0; Index < Indices.Num(); ++Index)
-		{
-			if (Indices[Index] > NumVertices)
-			{
-				NumVertices = Indices[Index];
-			}
-		}
-		NumVertices += 1;
-
-		TArray<uint32> OptimizedIndices;
-		OptimizedIndices.AddUninitialized(Indices.Num());
-		uint16 CacheSize = 32;
-		OptimizeFaces((uint8*)Indices.GetData(), Is32Bit, Indices.Num(), NumVertices, OptimizedIndices.GetData(), CacheSize);
-
-		if (Is32Bit)
-		{
-			FMemory::Memcpy(Indices.GetData(), OptimizedIndices.GetData(), Indices.Num() * sizeof(IndexDataType));
-		}
-		else
-		{
-			for (int32 I = 0; I < OptimizedIndices.Num(); ++I)
-			{
-				Indices[I] = (uint16)OptimizedIndices[I];
-			}
-		}
-	}
-}
-
 void FMeshUtilities::CacheOptimizeIndexBuffer(TArray<uint16>& Indices)
 {
-	if (bUsingNvTriStrip)
-	{
-		NvTriStrip::CacheOptimizeIndexBuffer(Indices);
-	}
-	else if (!bDisableTriangleOrderOptimization)
-	{
-		Forsyth::CacheOptimizeIndexBuffer(Indices);
-	}
+	BuildOptimizationThirdParty::CacheOptimizeIndexBuffer(Indices);
 }
 
 void FMeshUtilities::CacheOptimizeIndexBuffer(TArray<uint32>& Indices)
 {
-	if (bUsingNvTriStrip)
-	{
-		NvTriStrip::CacheOptimizeIndexBuffer(Indices);
-	}
-	else if (!bDisableTriangleOrderOptimization)
-	{
-		Forsyth::CacheOptimizeIndexBuffer(Indices);
-	}
-}
-
-/*------------------------------------------------------------------------------
-NVTessLib for computing adjacency used for tessellation.
-------------------------------------------------------------------------------*/
-
-/**
-* Provides static mesh render data to the NVIDIA tessellation library.
-*/
-class FStaticMeshNvRenderBuffer : public nv::RenderBuffer
-{
-public:
-
-	/** Construct from static mesh render buffers. */
-	FStaticMeshNvRenderBuffer(
-		const FPositionVertexBuffer& InPositionVertexBuffer,
-		const FStaticMeshVertexBuffer& InVertexBuffer,
-		const TArray<uint32>& Indices)
-		: PositionVertexBuffer(InPositionVertexBuffer)
-		, VertexBuffer(InVertexBuffer)
-	{
-		check(PositionVertexBuffer.GetNumVertices() == VertexBuffer.GetNumVertices());
-		mIb = new nv::IndexBuffer((void*)Indices.GetData(), nv::IBT_U32, Indices.Num(), false);
-	}
-
-	/** Retrieve the position and first texture coordinate of the specified index. */
-	virtual nv::Vertex getVertex(unsigned int Index) const
-	{
-		nv::Vertex Vertex;
-
-		check(Index < PositionVertexBuffer.GetNumVertices());
-
-		const FVector& Position = PositionVertexBuffer.VertexPosition(Index);
-		Vertex.pos.x = Position.X;
-		Vertex.pos.y = Position.Y;
-		Vertex.pos.z = Position.Z;
-
-		if (VertexBuffer.GetNumTexCoords())
-		{
-			const FVector2D UV = VertexBuffer.GetVertexUV(Index, 0);
-			Vertex.uv.x = UV.X;
-			Vertex.uv.y = UV.Y;
-		}
-		else
-		{
-			Vertex.uv.x = 0.0f;
-			Vertex.uv.y = 0.0f;
-		}
-
-		return Vertex;
-	}
-
-private:
-
-	/** The position vertex buffer for the static mesh. */
-	const FPositionVertexBuffer& PositionVertexBuffer;
-
-	/** The vertex buffer for the static mesh. */
-	const FStaticMeshVertexBuffer& VertexBuffer;
-
-	/** Copying is forbidden. */
-	FStaticMeshNvRenderBuffer(const FStaticMeshNvRenderBuffer&);
-	FStaticMeshNvRenderBuffer& operator=(const FStaticMeshNvRenderBuffer&);
-};
-
-/**
-* Provides skeletal mesh render data to the NVIDIA tessellation library.
-*/
-class FSkeletalMeshNvRenderBuffer : public nv::RenderBuffer
-{
-public:
-
-	/** Construct from static mesh render buffers. */
-	FSkeletalMeshNvRenderBuffer(
-		const TArray<FSoftSkinVertex>& InVertexBuffer,
-		const uint32 InTexCoordCount,
-		const TArray<uint32>& Indices)
-		: VertexBuffer(InVertexBuffer)
-		, TexCoordCount(InTexCoordCount)
-	{
-		mIb = new nv::IndexBuffer((void*)Indices.GetData(), nv::IBT_U32, Indices.Num(), false);
-	}
-
-	/** Retrieve the position and first texture coordinate of the specified index. */
-	virtual nv::Vertex getVertex(unsigned int Index) const
-	{
-		nv::Vertex Vertex;
-
-		check(Index < (unsigned int)VertexBuffer.Num());
-
-		const FSoftSkinVertex& SrcVertex = VertexBuffer[Index];
-
-		Vertex.pos.x = SrcVertex.Position.X;
-		Vertex.pos.y = SrcVertex.Position.Y;
-		Vertex.pos.z = SrcVertex.Position.Z;
-
-		if (TexCoordCount > 0)
-		{
-			Vertex.uv.x = SrcVertex.UVs[0].X;
-			Vertex.uv.y = SrcVertex.UVs[0].Y;
-		}
-		else
-		{
-			Vertex.uv.x = 0.0f;
-			Vertex.uv.y = 0.0f;
-		}
-
-		return Vertex;
-	}
-
-private:
-	/** The vertex buffer for the skeletal mesh. */
-	const TArray<FSoftSkinVertex>& VertexBuffer;
-	const uint32 TexCoordCount;
-
-	/** Copying is forbidden. */
-	FSkeletalMeshNvRenderBuffer(const FSkeletalMeshNvRenderBuffer&);
-	FSkeletalMeshNvRenderBuffer& operator=(const FSkeletalMeshNvRenderBuffer&);
-};
-
-static void BuildStaticAdjacencyIndexBuffer(
-	const FPositionVertexBuffer& PositionVertexBuffer,
-	const FStaticMeshVertexBuffer& VertexBuffer,
-	const TArray<uint32>& Indices,
-	TArray<uint32>& OutPnAenIndices
-	)
-{
-	if (Indices.Num())
-	{
-		FStaticMeshNvRenderBuffer StaticMeshRenderBuffer(PositionVertexBuffer, VertexBuffer, Indices);
-		nv::IndexBuffer* PnAENIndexBuffer = nv::tess::buildTessellationBuffer(&StaticMeshRenderBuffer, nv::DBM_PnAenDominantCorner, true);
-		check(PnAENIndexBuffer);
-		const int32 IndexCount = (int32)PnAENIndexBuffer->getLength();
-		OutPnAenIndices.Empty(IndexCount);
-		OutPnAenIndices.AddUninitialized(IndexCount);
-		for (int32 Index = 0; Index < IndexCount; ++Index)
-		{
-			OutPnAenIndices[Index] = (*PnAENIndexBuffer)[Index];
-		}
-		delete PnAENIndexBuffer;
-	}
-	else
-	{
-		OutPnAenIndices.Empty();
-	}
+	BuildOptimizationThirdParty::CacheOptimizeIndexBuffer(Indices);
 }
 
 void FMeshUtilities::BuildSkeletalAdjacencyIndexBuffer(
@@ -459,24 +135,7 @@ void FMeshUtilities::BuildSkeletalAdjacencyIndexBuffer(
 	TArray<uint32>& OutPnAenIndices
 	)
 {
-	if (Indices.Num())
-	{
-		FSkeletalMeshNvRenderBuffer SkeletalMeshRenderBuffer(VertexBuffer, TexCoordCount, Indices);
-		nv::IndexBuffer* PnAENIndexBuffer = nv::tess::buildTessellationBuffer(&SkeletalMeshRenderBuffer, nv::DBM_PnAenDominantCorner, true);
-		check(PnAENIndexBuffer);
-		const int32 IndexCount = (int32)PnAENIndexBuffer->getLength();
-		OutPnAenIndices.Empty(IndexCount);
-		OutPnAenIndices.AddUninitialized(IndexCount);
-		for (int32 Index = 0; Index < IndexCount; ++Index)
-		{
-			OutPnAenIndices[Index] = (*PnAENIndexBuffer)[Index];
-		}
-		delete PnAENIndexBuffer;
-	}
-	else
-	{
-		OutPnAenIndices.Empty();
-	}
+	BuildOptimizationThirdParty::NvTriStripHelper::BuildSkeletalAdjacencyIndexBuffer(VertexBuffer, TexCoordCount, Indices, OutPnAenIndices);
 }
 
 void FMeshUtilities::CalcBoneVertInfos(USkeletalMesh* SkeletalMesh, TArray<FBoneVertInfo>& Infos, bool bOnlyDominant)
@@ -3302,7 +2961,7 @@ public:
 			{
 				TArray<uint32> AdjacencyIndices;
 
-				BuildStaticAdjacencyIndexBuffer(
+				BuildOptimizationThirdParty::NvTriStripHelper::BuildStaticAdjacencyIndexBuffer(
 					LODModel.VertexBuffers.PositionVertexBuffer,
 					LODModel.VertexBuffers.StaticMeshVertexBuffer,
 					CombinedIndices,
@@ -5704,9 +5363,11 @@ void FMeshUtilities::StartupModule()
 
 	PropertyEditorModule.RegisterCustomClassLayout("ProxyLODMeshSimplificationSettings", FOnGetDetailCustomizationInstance::CreateStatic(&FProxyLODMeshSimplifcationSettingsCustomization::MakeInstance));
 
-	bDisableTriangleOrderOptimization = (CVarTriangleOrderOptimization.GetValueOnGameThread() == 2);
+	static const TConsoleVariableData<int32>* CVar = IConsoleManager::Get().FindTConsoleVariableDataInt(TEXT("r.TriangleOrderOptimization"));
+
+	bDisableTriangleOrderOptimization = (CVar->GetValueOnGameThread() == 2);
 
-	bUsingNvTriStrip = !bDisableTriangleOrderOptimization && (CVarTriangleOrderOptimization.GetValueOnGameThread() == 0);
+	bUsingNvTriStrip = !bDisableTriangleOrderOptimization && (CVar->GetValueOnGameThread() == 0);
 
 	IMeshReductionManagerModule& Module = FModuleManager::Get().LoadModuleChecked<IMeshReductionManagerModule>("MeshReductionInterface");
 	IMeshReduction* StaticMeshReduction = Module.GetStaticMeshReductionInterface();
diff --git a/Engine/Source/Runtime/Engine/Private/StaticMesh.cpp b/Engine/Source/Runtime/Engine/Private/StaticMesh.cpp
--- a/Engine/Source/Runtime/Engine/Private/StaticMesh.cpp
+++ b/Engine/Source/Runtime/Engine/Private/StaticMesh.cpp
@@ -1333,7 +1333,7 @@ FArchive& operator<<(FArchive& Ar, FMeshBuildSettings& BuildSettings)
 // differences, etc.) replace the version GUID below with a new one.
 // In case of merge conflicts with DDC versions, you MUST generate a new GUID
 // and set this new GUID as the version.                                       
-#define STATICMESH_DERIVEDDATA_VER TEXT("305A71049F594F5A8D73BE076F3DEF11")
+#define STATICMESH_DERIVEDDATA_VER TEXT("EAB28E38508F49109D98EED6C7CC9977")
 
 static const FString& GetStaticMeshDerivedDataVersion()
 {
@@ -2689,7 +2689,7 @@ void UStaticMesh::FixupMaterialSlotName()
 // differences, etc.) replace the version GUID below with a new one.
 // In case of merge conflicts with DDC versions, you MUST generate a new GUID
 // and set this new GUID as the version.                                       
-#define MESHDATAKEY_STATICMESH_DERIVEDDATA_VER TEXT("0329735AAE644B4AAC069ED412F0574B")
+#define MESHDATAKEY_STATICMESH_DERIVEDDATA_VER TEXT("97E0546FA0D747F4A1C78EC1EDF24CB3")
 
 static const FString& GetMeshDataKeyStaticMeshDerivedDataVersion()
 {
