Use new attribute array API for accessing FMeshDescription data.

diff --git a/Engine/Source/Developer/MeshDescriptionOperations/Private/MeshDescriptionOperations.cpp b/Engine/Source/Developer/MeshDescriptionOperations/Private/MeshDescriptionOperations.cpp
--- a/Engine/Source/Developer/MeshDescriptionOperations/Private/MeshDescriptionOperations.cpp
+++ b/Engine/Source/Developer/MeshDescriptionOperations/Private/MeshDescriptionOperations.cpp
@@ -1056,14 +1056,14 @@ void FMeshDescriptionOperations::FindOverlappingCorners(FOverlappingCorners& Out
 struct FLayoutUVMeshDescriptionView final : FLayoutUV::IMeshView
 {
 	FMeshDescription& MeshDescription;
-	const TVertexAttributeArray<FVector>& Positions;
-	const TVertexInstanceAttributeArray<FVector>& Normals;
-	const TVertexInstanceAttributeArray<FVector2D>& TexCoords;
+	const TVertexAttributesRef<FVector> Positions;
+	const TVertexInstanceAttributesRef<FVector> Normals;
+	const TVertexInstanceAttributesRef<FVector2D> TexCoords;
 
 	const uint32 SrcChannel;
 	const uint32 DstChannel;
 
-	TVertexInstanceAttributeArray<FVector2D>* OutputTexCoords = nullptr;
+	TVertexInstanceAttributesRef<FVector2D> OutputTexCoords;
 
 	uint32 NumIndices = 0;
 	TArray<int32> RemapVerts;
@@ -1071,9 +1071,9 @@ struct FLayoutUVMeshDescriptionView final : FLayoutUV::IMeshView
 
 	FLayoutUVMeshDescriptionView(FMeshDescription& InMeshDescription, uint32 InSrcChannel, uint32 InDstChannel) 
 		: MeshDescription(InMeshDescription)
-		, Positions(InMeshDescription.VertexAttributes().GetAttributes<FVector>(MeshAttribute::Vertex::Position))
-		, Normals(InMeshDescription.VertexInstanceAttributes().GetAttributes<FVector>(MeshAttribute::VertexInstance::Normal))
-		, TexCoords(MeshDescription.VertexInstanceAttributes().GetAttributes<FVector2D>(MeshAttribute::VertexInstance::TextureCoordinate, InSrcChannel))
+		, Positions(InMeshDescription.VertexAttributes().GetAttributesRef<FVector>(MeshAttribute::Vertex::Position))
+		, Normals(InMeshDescription.VertexInstanceAttributes().GetAttributesRef<FVector>(MeshAttribute::VertexInstance::Normal))
+		, TexCoords(InMeshDescription.VertexInstanceAttributes().GetAttributesRef<FVector2D>(MeshAttribute::VertexInstance::TextureCoordinate))
 		, SrcChannel(InSrcChannel)
 		, DstChannel(InDstChannel)
 	{
@@ -1099,7 +1099,7 @@ struct FLayoutUVMeshDescriptionView final : FLayoutUV::IMeshView
 				{
 					const FVertexInstanceID VertexInstanceID = MeshTriangle.GetVertexInstanceID(Corner);
 
-					FlattenedTexCoords[WedgeIndex] = TexCoords[VertexInstanceID];
+					FlattenedTexCoords[WedgeIndex] = TexCoords.Get(VertexInstanceID, SrcChannel);
 					RemapVerts[WedgeIndex] = VertexInstanceID.GetValue();
 					++WedgeIndex;
 				}
@@ -1137,13 +1137,13 @@ struct FLayoutUVMeshDescriptionView final : FLayoutUV::IMeshView
 			ensure(false);	// not expecting it to get here
 		}
 
-		OutputTexCoords = &MeshDescription.VertexInstanceAttributes().GetAttributes<FVector2D>(MeshAttribute::VertexInstance::TextureCoordinate, DstChannel);
+		OutputTexCoords = MeshDescription.VertexInstanceAttributes().GetAttributesRef<FVector2D>(MeshAttribute::VertexInstance::TextureCoordinate);
 	}
 
 	void SetOutputTexcoord(uint32 Index, const FVector2D& Value) override
 	{
 		const FVertexInstanceID VertexInstanceID(RemapVerts[Index]);
-		(*OutputTexCoords)[VertexInstanceID] = Value;
+		OutputTexCoords.Set(VertexInstanceID, DstChannel, Value);
 	}
 };
 
