Fixed deprecated MeshDescription calls (merged with Yuriy's changes).

diff --git a/Engine/Source/Developer/MeshDescriptionOperations/Private/MeshDescriptionOperations.cpp b/Engine/Source/Developer/MeshDescriptionOperations/Private/MeshDescriptionOperations.cpp
--- a/Engine/Source/Developer/MeshDescriptionOperations/Private/MeshDescriptionOperations.cpp
+++ b/Engine/Source/Developer/MeshDescriptionOperations/Private/MeshDescriptionOperations.cpp
@@ -1056,15 +1056,13 @@ void FMeshDescriptionOperations::FindOverlappingCorners(FOverlappingCorners& Out
 struct FLayoutUVMeshDescriptionView final : FLayoutUV::IMeshView
 {
 	FMeshDescription& MeshDescription;
-	const TVertexAttributesRef<FVector> Positions;
-	const TVertexInstanceAttributesRef<FVector> Normals;
-	const TVertexInstanceAttributesRef<FVector2D> TexCoords;
+	TVertexAttributesConstRef<FVector> Positions;
+	TVertexInstanceAttributesConstRef<FVector> Normals;
+	TVertexInstanceAttributesRef<FVector2D> TexCoords;
 
 	const uint32 SrcChannel;
 	const uint32 DstChannel;
 
-	TVertexInstanceAttributesRef<FVector2D> OutputTexCoords;
-
 	uint32 NumIndices = 0;
 	TArray<int32> RemapVerts;
 	TArray<FVector2D> FlattenedTexCoords;
@@ -1130,20 +1128,18 @@ struct FLayoutUVMeshDescriptionView final : FLayoutUV::IMeshView
 	void InitOutputTexcoords(uint32 Num) override
 	{
 		// If current DstChannel is out of range of the number of UVs defined by the mesh description, change the index count accordingly
-		const uint32 NumUVs = MeshDescription.VertexInstanceAttributes().GetAttributeIndexCount<FVector2D>(MeshAttribute::VertexInstance::TextureCoordinate);
+		const uint32 NumUVs = TexCoords.GetNumIndices();
 		if (DstChannel >= NumUVs)
 		{
-			MeshDescription.VertexInstanceAttributes().SetAttributeIndexCount<FVector2D>(MeshAttribute::VertexInstance::TextureCoordinate, DstChannel + 1);
+			TexCoords.SetNumIndices(DstChannel + 1);
 			ensure(false);	// not expecting it to get here
 		}
-
-		OutputTexCoords = MeshDescription.VertexInstanceAttributes().GetAttributesRef<FVector2D>(MeshAttribute::VertexInstance::TextureCoordinate);
 	}
 
 	void SetOutputTexcoord(uint32 Index, const FVector2D& Value) override
 	{
 		const FVertexInstanceID VertexInstanceID(RemapVerts[Index]);
-		OutputTexCoords.Set(VertexInstanceID, DstChannel, Value);
+		TexCoords.Set(VertexInstanceID, DstChannel, Value);
 	}
 };
 
