Fix the GenerateUniqueUVsForStaticMesh when there is some duplicate polygon and we remove them, we must avoid Compacting the Polygon so the Remap VertexInstanceID is still valid.
There was a second crash after this one where the code did not call RegisterMeshAttribute on a newly created FMeshDescription.

diff --git a/Engine/Source/Developer/MeshDescriptionOperations/Private/MeshDescriptionOperations.cpp b/Engine/Source/Developer/MeshDescriptionOperations/Private/MeshDescriptionOperations.cpp
--- a/Engine/Source/Developer/MeshDescriptionOperations/Private/MeshDescriptionOperations.cpp
+++ b/Engine/Source/Developer/MeshDescriptionOperations/Private/MeshDescriptionOperations.cpp
@@ -1553,9 +1553,7 @@ bool FMeshDescriptionOperations::GenerateUniqueUVsForStaticMesh(const FMeshDescr
 			{
 				DuplicateMeshDescription.DeleteVertex(VertexID);
 			}
-			//Compact and Remap IDs so we have clean ID from 0 to n since we just erase some polygons
-			FElementIDRemappings RemappingInfos;
-			DuplicateMeshDescription.Compact(RemappingInfos);
+			//Avoid compacting the DuplicateMeshDescription, since the remap of the VertexInstaceID will not be good anymore
 		}
 	}
 	// Find overlapping corners for UV generator. Allow some threshold - this should not produce any error in a case if resulting
diff --git a/Engine/Source/Developer/MeshMergeUtilities/Private/MeshMergeHelpers.cpp b/Engine/Source/Developer/MeshMergeUtilities/Private/MeshMergeHelpers.cpp
--- a/Engine/Source/Developer/MeshMergeUtilities/Private/MeshMergeHelpers.cpp
+++ b/Engine/Source/Developer/MeshMergeUtilities/Private/MeshMergeHelpers.cpp
@@ -1303,6 +1303,7 @@ void FMeshMergeHelpers::MergeImpostersToRawMesh(TArray<const UStaticMeshComponen
 
 			// Retrieve mesh data in FMeshDescription form
 			FMeshDescription ImposterMesh;
+			UStaticMesh::RegisterMeshAttributes(ImposterMesh);
 			FMeshMergeHelpers::RetrieveMesh(Component, LODIndex, ImposterMesh, false);
 
 			// Retrieve the sections, we're expect 1 for imposter meshes
