Enterprise UV Generation tool - Integrating CL 6893582, 6954109, 6957407, 6962916, 6969266, 6969351, 7004882, 7047829

diff --git a/Engine/Source/Developer/MeshDescriptionOperations/Private/MeshDescriptionOperations.cpp b/Engine/Source/Developer/MeshDescriptionOperations/Private/MeshDescriptionOperations.cpp
--- a/Engine/Source/Developer/MeshDescriptionOperations/Private/MeshDescriptionOperations.cpp
+++ b/Engine/Source/Developer/MeshDescriptionOperations/Private/MeshDescriptionOperations.cpp
@@ -1716,7 +1716,7 @@ void FMeshDescriptionOperations::GenerateCylindricalUV(FMeshDescription& MeshDes
 		int32 NumInstances = VertexInstances.Num();
 		if (NumInstances >= 2)
 		{
-			for (int32 StartIndex = 1; StartIndex < NumInstances; ++StartIndex)
+			for (int32 StartIndex = 0; StartIndex < NumInstances; ++StartIndex)
 			{
 				int32 EndIndex = StartIndex + 1;
 				if (EndIndex >= NumInstances)
@@ -1734,11 +1734,17 @@ void FMeshDescriptionOperations::GenerateCylindricalUV(FMeshDescription& MeshDes
 					// Fix the U coordinate to get the texture go counterclockwise
 					if (EndUV.X > Threshold)
 					{
-						EndUV.X -= 1.f;
+						if (EndUV.X >= 1.f)
+						{
+							EndUV.X -= 1.f;
+						}
 					}
 					else
 					{
-						EndUV.X += 1.f;
+						if (EndUV.X <= 0)
+						{
+							EndUV.X += 1.f;
+						}
 					}
 				}
 			}
