Copied 4369604 JeanLuc.Corenthin
- Fixed indexing issue with FOverlappingCorners object created using a FmeshDescription object:
    - Changed logic in building of FOverlappingCorners to be based on index of wedge, number of wedges equal 3 times the number of triangles like in FRawMesh, instead of index of vertex instances.
   - Updated BuildVertexBuffer to take this new logic in account.
- Fixed crash in FMeshDescriptionOperations::CreateMikktTangents:
   - MikkTSpace sdk does not handle properly faces with 0 vertices which occur with mesh description with invalid polygons.

diff --git a/Engine/Source/Developer/MeshDescriptionOperations/Private/MeshDescriptionOperations.cpp b/Engine/Source/Developer/MeshDescriptionOperations/Private/MeshDescriptionOperations.cpp
--- a/Engine/Source/Developer/MeshDescriptionOperations/Private/MeshDescriptionOperations.cpp
+++ b/Engine/Source/Developer/MeshDescriptionOperations/Private/MeshDescriptionOperations.cpp
@@ -926,6 +926,14 @@ namespace MeshDescriptionMikktSpaceInterface
 
 void FMeshDescriptionOperations::CreateMikktTangents(FMeshDescription& MeshDescription, FMeshDescriptionOperations::ETangentOptions TangentOptions)
 {
+	// The Mikkt interface does not handle properly polygon array with 'holes'
+	// Compact mesh description if this is the case
+	if (MeshDescription.Polygons().Num() != MeshDescription.Polygons().GetArraySize())
+	{
+		FElementIDRemappings Remappings;
+		MeshDescription.Compact(Remappings);
+	}
+
 	bool bIgnoreDegenerateTriangles = (TangentOptions & FMeshDescriptionOperations::ETangentOptions::IgnoreDegenerateTriangles) != 0;
 
 	// we can use mikktspace to calculate the tangents
@@ -1017,7 +1025,15 @@ void FMeshDescriptionOperations::FindOverlappingCorners(FOverlappingCorners& Out
 	const FVertexInstanceArray& VertexInstanceArray = MeshDescription.VertexInstances();
 	const FVertexArray& VertexArray = MeshDescription.Vertices();
 
-	const int32 NumWedges = VertexInstanceArray.Num();
+	int32 NumWedges = 0;
+	for (const FPolygonID PolygonID : MeshDescription.Polygons().GetElementIDs())
+	{
+		const TArray<FMeshTriangle>& Triangles = MeshDescription.GetPolygonTriangles(PolygonID);
+		for (const FMeshTriangle MeshTriangle : Triangles)
+		{
+			NumWedges += 3;
+		}
+	}
 
 	// Empty the old data and reserve space for new
 	OutOverlappingCorners.Init(NumWedges);
@@ -1028,9 +1044,19 @@ void FMeshDescriptionOperations::FindOverlappingCorners(FOverlappingCorners& Out
 
 	TVertexAttributesConstRef<FVector> VertexPositions = MeshDescription.VertexAttributes().GetAttributesRef<FVector>(MeshAttribute::Vertex::Position);
 
-	for (const FVertexInstanceID VertexInstanceID : VertexInstanceArray.GetElementIDs())
+	int32 WedgeIndex = 0;
+	for (const FPolygonID PolygonID : MeshDescription.Polygons().GetElementIDs())
 	{
-		new(VertIndexAndZ)MeshDescriptionOperationNamespace::FIndexAndZ(VertexInstanceID.GetValue(), VertexPositions[MeshDescription.GetVertexInstanceVertex(VertexInstanceID)]);
+		const TArray<FMeshTriangle>& Triangles = MeshDescription.GetPolygonTriangles(PolygonID);
+		for (const FMeshTriangle MeshTriangle : Triangles)
+		{
+			for (int32 Corner = 0; Corner < 3; ++Corner)
+			{
+				const FVertexInstanceID VertexInstanceID = MeshTriangle.GetVertexInstanceID(Corner);
+				new(VertIndexAndZ)MeshDescriptionOperationNamespace::FIndexAndZ(WedgeIndex, VertexPositions[MeshDescription.GetVertexInstanceVertex(VertexInstanceID)]);
+				++WedgeIndex;
+			}
+		}
 	}
 
 	// Sort the vertices by z value
