Fix crash in MeshDescriptionOperations compute normals, the UVs use as keys had to be sanitized in case of -0.0 versus 0.0

diff --git a/Engine/Source/Developer/MeshDescriptionOperations/Private/MeshDescriptionOperations.cpp b/Engine/Source/Developer/MeshDescriptionOperations/Private/MeshDescriptionOperations.cpp
--- a/Engine/Source/Developer/MeshDescriptionOperations/Private/MeshDescriptionOperations.cpp
+++ b/Engine/Source/Developer/MeshDescriptionOperations/Private/MeshDescriptionOperations.cpp
@@ -682,6 +682,18 @@ void FMeshDescriptionOperations::CreatePolygonNTB(FMeshDescription& MeshDescript
 	}
 }
 
+struct FNTBGroupKeyFuncs : public TDefaultMapKeyFuncs<FVector2D, FVector, false>
+{
+	//We need to sanitize the key here to make sure -0.0f fall on the same hash then 0.0f
+	static FORCEINLINE_DEBUGGABLE uint32 GetKeyHash(KeyInitType Key)
+	{
+		FVector2D TmpKey;
+		TmpKey.X = FMath::IsNearlyZero(Key.X) ? 0.0f : Key.X;
+		TmpKey.Y = FMath::IsNearlyZero(Key.Y) ? 0.0f : Key.Y;
+		return FCrc::MemCrc32(&TmpKey, sizeof(TmpKey));
+	}
+};
+
 void FMeshDescriptionOperations::CreateNormals(FMeshDescription& MeshDescription, FMeshDescriptionOperations::ETangentOptions TangentOptions, bool bComputeTangent)
 {
 	//For each vertex compute the normals for every connected edges that are smooth betwween hard edges
@@ -811,8 +823,8 @@ void FMeshDescriptionOperations::CreateNormals(FMeshDescription& MeshDescription
 		for (const TArray<FPolygonID>& Group : Groups)
 		{
 			//Compute tangents data
-			TMap<FVector2D, FVector> GroupTangent;
-			TMap<FVector2D, FVector> GroupBiNormal;
+			TMap<FVector2D, FVector,FDefaultSetAllocator, FNTBGroupKeyFuncs> GroupTangent;
+			TMap<FVector2D, FVector, FDefaultSetAllocator, FNTBGroupKeyFuncs> GroupBiNormal;
 
 			TArray<FVertexInstanceID> VertexInstanceInGroup;
 			FVector GroupNormal(0.0f);
