Minor fixes on MeshDescriptionOperations:
- avoid needless copies of FMeshTriangle

diff --git a/Engine/Source/Developer/MeshDescriptionOperations/Private/MeshDescriptionOperations.cpp b/Engine/Source/Developer/MeshDescriptionOperations/Private/MeshDescriptionOperations.cpp
--- a/Engine/Source/Developer/MeshDescriptionOperations/Private/MeshDescriptionOperations.cpp
+++ b/Engine/Source/Developer/MeshDescriptionOperations/Private/MeshDescriptionOperations.cpp
@@ -916,7 +916,7 @@ void FMeshDescriptionOperations::CreateNormals(FMeshDescription& MeshDescription
 	//          \  ||  /
 	//       G  -- ** -- C
 	//          // |  \
-		//         F   E    D
+	//         F   E    D
 //
 // The double ** are the vertex, the double line are hard edges, the single line are soft edge.
 // A and F are hard, all other edges are soft. The goal is to compute two average normals one from
@@ -1256,11 +1256,7 @@ void FMeshDescriptionOperations::FindOverlappingCorners(FOverlappingCorners& Out
 	int32 NumWedges = 0;
 	for (const FPolygonID PolygonID : MeshDescription.Polygons().GetElementIDs())
 	{
-		const TArray<FMeshTriangle>& Triangles = MeshDescription.GetPolygonTriangles(PolygonID);
-		for (const FMeshTriangle MeshTriangle : Triangles)
-		{
-			NumWedges += 3;
-		}
+		NumWedges += 3 * MeshDescription.GetPolygonTriangles(PolygonID).Num();
 	}
 
 	// Empty the old data and reserve space for new
@@ -1276,7 +1272,7 @@ void FMeshDescriptionOperations::FindOverlappingCorners(FOverlappingCorners& Out
 	for (const FPolygonID PolygonID : MeshDescription.Polygons().GetElementIDs())
 	{
 		const TArray<FMeshTriangle>& Triangles = MeshDescription.GetPolygonTriangles(PolygonID);
-		for (const FMeshTriangle MeshTriangle : Triangles)
+		for (const FMeshTriangle& MeshTriangle : Triangles)
 		{
 			for (int32 Corner = 0; Corner < 3; ++Corner)
 			{
