[DataSmith] RawMesh to MeshDescription conversion: Algo part

diff --git a/Engine/Source/Developer/MeshDescriptionOperations/Private/MeshDescriptionOperations.cpp b/Engine/Source/Developer/MeshDescriptionOperations/Private/MeshDescriptionOperations.cpp
--- a/Engine/Source/Developer/MeshDescriptionOperations/Private/MeshDescriptionOperations.cpp
+++ b/Engine/Source/Developer/MeshDescriptionOperations/Private/MeshDescriptionOperations.cpp
@@ -477,11 +477,14 @@ void FMeshDescriptionOperations::ConvertFromRawMesh(const FRawMesh& SourceRawMes
 		{
 			FPolygonGroupID PolygonGroupID(MaterialIndex);
 			DestinationMeshDescription.CreatePolygonGroupWithID(PolygonGroupID);
-			PolygonGroupImportedMaterialSlotNames[PolygonGroupID] = FName(*FString::Printf(TEXT("MaterialSlot_%d"), MaterialIndex));
 			if (MaterialMap.Contains(MaterialIndex))
 			{
 				PolygonGroupImportedMaterialSlotNames[PolygonGroupID] = MaterialMap[MaterialIndex];
 			}
+			else
+			{
+				PolygonGroupImportedMaterialSlotNames[PolygonGroupID] = FName(*FString::Printf(TEXT("MaterialSlot_%d"), MaterialIndex));
+			}
 			PolygonGroups.Add(PolygonGroupID);
 			MaterialIndexToPolygonGroup.Add(MaterialIndex, PolygonGroupID);
 		}
@@ -497,7 +500,7 @@ void FMeshDescriptionOperations::ConvertFromRawMesh(const FRawMesh& SourceRawMes
 		{
 			continue;
 		}
-		
+
 		//PolygonGroup
 		FPolygonGroupID PolygonGroupID = FPolygonGroupID::Invalid;
 		FName PolygonGroupImportedMaterialSlotName = NAME_None;
@@ -518,7 +521,7 @@ void FMeshDescriptionOperations::ConvertFromRawMesh(const FRawMesh& SourceRawMes
 				}
 			}
 		}
-		
+
 		if (PolygonGroupID == FPolygonGroupID::Invalid)
 		{
 			PolygonGroupID = DestinationMeshDescription.CreatePolygonGroup();
@@ -581,7 +584,7 @@ void FMeshDescriptionOperations::ConvertFromRawMesh(const FRawMesh& SourceRawMes
 			NewTriangle.SetVertexInstanceID(Corner, VertexInstanceID);
 		}
 	}
-	
+
 	ConvertSmoothGroupToHardEdges(SourceRawMesh.FaceSmoothingMasks, DestinationMeshDescription);
 
 	//Create the missing normals and tangents, should we use Mikkt space for tangent???
@@ -619,7 +622,7 @@ void FMeshDescriptionOperations::AppendMeshDescription(const FMeshDescription& S
 
 	//PolygonGroup Attributes
 	TPolygonGroupAttributesConstRef<FName> SourceImportedMaterialSlotNames = SourceMesh.PolygonGroupAttributes().GetAttributesRef<FName>(MeshAttribute::PolygonGroup::ImportedMaterialSlotName);
-	
+
 	TPolygonGroupAttributesRef<FName> TargetImportedMaterialSlotNames = TargetMesh.PolygonGroupAttributes().GetAttributesRef<FName>(MeshAttribute::PolygonGroup::ImportedMaterialSlotName);
 
 	//VertexInstance Attributes
@@ -685,6 +688,17 @@ void FMeshDescriptionOperations::AppendMeshDescription(const FMeshDescription& S
 		SourceVertexIDRemap.Add(SourceVertexID, TargetVertexID);
 	}
 
+	// Transform vertices properties
+	if (AppendSettings.MeshTransform)
+	{
+		const FTransform& Transform = AppendSettings.MeshTransform.GetValue();
+		for (const TPair<FVertexID, FVertexID>& VertexIDPair : SourceVertexIDRemap)
+		{
+			FVector& Position = TargetVertexPositions[VertexIDPair.Value];
+			Position = Transform.TransformPosition(Position);
+		}
+	}
+
 	//Edges
 	TMap<FEdgeID, FEdgeID> SourceEdgeIDRemap;
 	SourceEdgeIDRemap.Reserve(SourceMesh.Edges().Num());
@@ -704,7 +718,7 @@ void FMeshDescriptionOperations::AppendMeshDescription(const FMeshDescription& S
 	{
 		FVertexInstanceID TargetVertexInstanceID = TargetMesh.CreateVertexInstance(SourceVertexIDRemap[SourceMesh.GetVertexInstanceVertex(SourceVertexInstanceID)]);
 		SourceVertexInstanceIDRemap.Add(SourceVertexInstanceID, TargetVertexInstanceID);
-		
+
 		TargetVertexInstanceNormals[TargetVertexInstanceID] = SourceVertexInstanceNormals[SourceVertexInstanceID];
 		TargetVertexInstanceTangents[TargetVertexInstanceID] = SourceVertexInstanceTangents[SourceVertexInstanceID];
 		TargetVertexInstanceBinormalSigns[TargetVertexInstanceID] = SourceVertexInstanceBinormalSigns[SourceVertexInstanceID];
@@ -720,6 +734,26 @@ void FMeshDescriptionOperations::AppendMeshDescription(const FMeshDescription& S
 		}
 	}
 
+	// Transform vertex instances properties
+	if (AppendSettings.MeshTransform)
+	{
+		const FTransform& Transform = AppendSettings.MeshTransform.GetValue();
+		bool bFlipBinormal = Transform.GetDeterminant() < 0;
+		float BinormalSignsFactor = bFlipBinormal ? -1.f : 1.f;
+		for (const TPair<FVertexInstanceID, FVertexInstanceID>& VertexInstanceIDPair : SourceVertexInstanceIDRemap)
+		{
+			FVertexInstanceID InstanceID = VertexInstanceIDPair.Value;
+
+			FVector& Normal = TargetVertexInstanceNormals[InstanceID];
+			Normal = Transform.TransformVectorNoScale(Normal);
+
+			FVector& Tangent = TargetVertexInstanceTangents[InstanceID];
+			Tangent = Transform.TransformVectorNoScale(Tangent);
+
+			TargetVertexInstanceBinormalSigns[InstanceID] *= BinormalSignsFactor;
+		}
+	}
+
 	//Polygons
 	for (const FPolygonID& SourcePolygonID : SourceMesh.Polygons().GetElementIDs())
 	{
@@ -917,17 +951,17 @@ void FMeshDescriptionOperations::CreateNormals(FMeshDescription& MeshDescription
 	//       G  -- ** -- C
 	//          // |  \
 	//         F   E    D
-//
-// The double ** are the vertex, the double line are hard edges, the single line are soft edge.
-// A and F are hard, all other edges are soft. The goal is to compute two average normals one from
-// A to F and a second from F to A. Then we can set the vertex instance normals accordingly.
-// First normal(A to F) = Normalize(A+B+C+D+E+F)
-// Second normal(F to A) = Normalize(F+G+H+A)
-// We found the connected edge using the triangle that share edges
-
-// @todo: provide an option to weight each contributing polygon normal according to the size of
-// the angle it makes with the vertex being calculated. This means that triangulated faces whose
-// internal edge meets the vertex doesn't get undue extra weight.
+	//
+	// The double ** are the vertex, the double line are hard edges, the single line are soft edge.
+	// A and F are hard, all other edges are soft. The goal is to compute two average normals one from
+	// A to F and a second from F to A. Then we can set the vertex instance normals accordingly.
+	// First normal(A to F) = Normalize(A+B+C+D+E+F)
+	// Second normal(F to A) = Normalize(F+G+H+A)
+	// We found the connected edge using the triangle that share edges
+
+	// @todo: provide an option to weight each contributing polygon normal according to the size of
+	// the angle it makes with the vertex being calculated. This means that triangulated faces whose
+	// internal edge meets the vertex doesn't get undue extra weight.
 
 	const TVertexInstanceAttributesRef<FVector2D> VertexUVs = MeshDescription.VertexInstanceAttributes().GetAttributesRef<FVector2D>(MeshAttribute::VertexInstance::TextureCoordinate);
 	TVertexInstanceAttributesRef<FVector> VertexNormals = MeshDescription.VertexInstanceAttributes().GetAttributesRef<FVector>(MeshAttribute::VertexInstance::Normal);
@@ -1048,7 +1082,7 @@ void FMeshDescriptionOperations::CreateNormals(FMeshDescription& MeshDescription
 				FVector PolyNormal = PolygonNormals[PolygonID];
 				FVector PolyTangent = PolygonTangents[PolygonID];
 				FVector PolyBinormal = PolygonBinormals[PolygonID];
-				
+
 				ConsumedPolygon.Add(PolygonID);
 				VertexInstanceInGroup.Add(VertexInfoMap[PolygonID].VertexInstanceID);
 				if (!PolyNormal.IsNearlyZero(SMALL_NUMBER) && !PolyNormal.ContainsNaN())
@@ -1323,7 +1357,7 @@ struct FLayoutUVMeshDescriptionView final : FLayoutUV::IMeshView
 	TArray<int32> RemapVerts;
 	TArray<FVector2D> FlattenedTexCoords;
 
-	FLayoutUVMeshDescriptionView(FMeshDescription& InMeshDescription, uint32 InSrcChannel, uint32 InDstChannel) 
+	FLayoutUVMeshDescriptionView(FMeshDescription& InMeshDescription, uint32 InSrcChannel, uint32 InDstChannel)
 		: MeshDescription(InMeshDescription)
 		, Positions(InMeshDescription.VertexAttributes().GetAttributesRef<FVector>(MeshAttribute::Vertex::Position))
 		, Normals(InMeshDescription.VertexInstanceAttributes().GetAttributesRef<FVector>(MeshAttribute::VertexInstance::Normal))
@@ -1364,14 +1398,14 @@ struct FLayoutUVMeshDescriptionView final : FLayoutUV::IMeshView
 	uint32 GetNumIndices() const override { return NumIndices; }
 
 	FVector GetPosition(uint32 Index) const override
-	{ 
+	{
 		FVertexInstanceID VertexInstanceID(RemapVerts[Index]);
 		FVertexID VertexID = MeshDescription.GetVertexInstanceVertex(VertexInstanceID);
 		return Positions[VertexID];
 	}
 
 	FVector GetNormal(uint32 Index) const override
-	{ 
+	{
 		FVertexInstanceID VertexInstanceID(RemapVerts[Index]);
 		return Normals[VertexInstanceID];
 	}
diff --git a/Engine/Source/Developer/MeshDescriptionOperations/Public/MeshDescriptionOperations.h b/Engine/Source/Developer/MeshDescriptionOperations/Public/MeshDescriptionOperations.h
--- a/Engine/Source/Developer/MeshDescriptionOperations/Public/MeshDescriptionOperations.h
+++ b/Engine/Source/Developer/MeshDescriptionOperations/Public/MeshDescriptionOperations.h
@@ -48,6 +48,7 @@ public:
 		FAppendPolygonGroupsDelegate PolygonGroupsDelegate;
 		bool bMergeVertexColor;
 		FVector MergedAssetPivot;
+		TOptional<FTransform> MeshTransform; // Apply a transformation on source mesh (see MeshTransform)
 	};
 
 	static void AppendMeshDescription(const FMeshDescription& SourceMesh, FMeshDescription& TargetMesh, FAppendSettings& AppendSettings);
@@ -62,10 +63,10 @@ public:
 	 * It also remove the degenerated polygon from the mesh description
 	 */
 	static void CreatePolygonNTB(FMeshDescription& MeshDescription, float ComparisonThreshold);
-	
+
 	/** Compute normal, tangent and Bi-Normal(only if bComputeTangent is true) for every vertex in the mesh description. */
 	static void CreateNormals(FMeshDescription& MeshDescription, ETangentOptions TangentOptions, bool bComputeTangent);
-	
+
 	/** Compute tangent and Bi-Normal using mikkt space for every vertex in the mesh description. */
 	static void CreateMikktTangents(FMeshDescription& MeshDescription, ETangentOptions TangentOptions);
 
