Fix the ConvertToRawMesh keep the section order use in the source MeshDescription.
Update one asset affect by the bug and update the screenshot tests

diff --git a/Engine/Source/Developer/MeshDescriptionOperations/Private/MeshDescriptionOperations.cpp b/Engine/Source/Developer/MeshDescriptionOperations/Private/MeshDescriptionOperations.cpp
--- a/Engine/Source/Developer/MeshDescriptionOperations/Private/MeshDescriptionOperations.cpp
+++ b/Engine/Source/Developer/MeshDescriptionOperations/Private/MeshDescriptionOperations.cpp
@@ -475,15 +475,14 @@ void FMeshDescriptionOperations::ConvertFromRawMesh(const FRawMesh& SourceRawMes
 	{
 		if (!MaterialIndexToPolygonGroup.Contains(MaterialIndex))
 		{
-			FPolygonGroupID PolygonGroupID(MaterialIndex);
-			DestinationMeshDescription.CreatePolygonGroupWithID(PolygonGroupID);
+			FPolygonGroupID PolygonGroupID = DestinationMeshDescription.CreatePolygonGroup();
 			if (MaterialMap.Contains(MaterialIndex))
 			{
 				PolygonGroupImportedMaterialSlotNames[PolygonGroupID] = MaterialMap[MaterialIndex];
 			}
 			else
 			{
-				PolygonGroupImportedMaterialSlotNames[PolygonGroupID] = FName(*FString::Printf(TEXT("MaterialSlot_%d"), MaterialIndex));
+				PolygonGroupImportedMaterialSlotNames[PolygonGroupID] = FName(*FString::Printf(TEXT("MaterialSlot_%d"), PolygonGroupID.GetValue()));
 			}
 			PolygonGroups.Add(PolygonGroupID);
 			MaterialIndexToPolygonGroup.Add(MaterialIndex, PolygonGroupID);
diff --git a/Engine/Source/Runtime/Engine/Private/StaticMesh.cpp b/Engine/Source/Runtime/Engine/Private/StaticMesh.cpp
--- a/Engine/Source/Runtime/Engine/Private/StaticMesh.cpp
+++ b/Engine/Source/Runtime/Engine/Private/StaticMesh.cpp
@@ -3025,7 +3025,7 @@ void UStaticMesh::FixupMaterialSlotName()
 // differences, etc.) replace the version GUID below with a new one.
 // In case of merge conflicts with DDC versions, you MUST generate a new GUID
 // and set this new GUID as the version.                                       
-#define MESHDATAKEY_STATICMESH_DERIVEDDATA_VER TEXT("7B4DCBC8C34E4C6E97A9A84BA61CADCD")
+#define MESHDATAKEY_STATICMESH_DERIVEDDATA_VER TEXT("0386934AE700020EBBC4DD45C0AA3C49")
 
 static const FString& GetMeshDataKeyStaticMeshDerivedDataVersion()
 {
