Merge conflict.
alexis.matte, please merge this change by hand.
More info at https://robomerge.epicgames.net#RELEASE

/src/ROBOMERGE_RELEASE_Main/EngineTest/Test/Screenshots/ActorMerging/ActorMerging_SimpleMeshMerge_LOD_1_None/Windows/D3D11_SM5/22c96594569807222f8f1a0ad282ed19.json - merging //UE4/Release-4.22/EngineTest/Test/Screenshots/ActorMerging/ActorMerging_SimpleMeshMerge_LOD_1_None/Windows/D3D11_SM5/22c96594569807222f8f1a0ad282ed19.json#3
/src/ROBOMERGE_RELEASE_Main/EngineTest/Test/Screenshots/ActorMerging/ActorMerging_SimpleMeshMerge_LOD_1_None/Windows/D3D11_SM5/22c96594569807222f8f1a0ad282ed19.png - vs //UE4/Release-4.22/EngineTest/Test/Screenshots/ActorMerging/ActorMerging_SimpleMeshMerge_LOD_1_None/Windows/D3D11_SM5/22c96594569807222f8f1a0ad282ed19.png#3
/src/ROBOMERGE_RELEASE_Main/EngineTest/Test/Screenshots/ActorMerging/ActorMerging_SimpleMeshMerge_LOD_1_None/Windows/D3D11_SM5/3fcc508dbd2d03ed4ef4a56e129b984a.json - resolving delete from //UE4/Release-4.22/EngineTest/Test/Screenshots/ActorMerging/ActorMerging_SimpleMeshMerge_LOD_1_None/Windows/D3D11_SM5/3fcc508dbd2d03ed4ef4a56e129b984a.json#2
Delete resolve:
at: delete
ay: ignore
/src/ROBOMERGE_RELEASE_Main/EngineTest/Test/Screenshots/ActorMerging/ActorMerging_SimpleMeshMerge_LOD_1_None/Windows/D3D11_SM5/3fcc508dbd2d03ed4ef4a56e129b984a.png - resolving delete from //UE4/Release-4.22/EngineTest/Test/Screenshots/ActorMerging/ActorMerging_SimpleMeshMerge_LOD_1_None/Windows/D3D11_SM5/3fcc508dbd2d03ed4ef4a56e129b984a.png#2
Delete resolve:
at: delete
ay: ignore
/src/ROBOMERGE_RELEASE_Main/EngineTest/Test/Screenshots/ActorMerging/ActorMerging_SimpleMeshMerge_LOD_2_None/Windows/D3D11_SM5/22c96594569807222f8f1a0ad282ed19.json - merging //UE4/Release-4.22/EngineTest/Test/Screenshots/ActorMerging/ActorMerging_SimpleMeshMerge_LOD_2_None/Windows/D3D11_SM5/22c96594569807222f8f1a0ad282ed19.json#3
/src/ROBOMERGE_RELEASE_Main/EngineTest/Test/Screenshots/ActorMerging/ActorMerging_SimpleMeshMerge_LOD_2_None/Windows/D3D11_SM5/22c96594569807222f8f1a0ad282ed19.png - vs //UE4/Release-4.22/EngineTest/Test/Screenshots/ActorMerging/ActorMerging_SimpleMeshMerge_LOD_2_None/Windows/D3D11_SM5/22c96594569807222f8f1a0ad282ed19.png#3
/src/ROBOMERGE_RELEASE_Main/EngineTest/Test/Screenshots/ActorMerging/ActorMerging_SimpleMeshMerge_LOD_2_None/Windows/D3D11_SM5/3fcc508dbd2d03ed4ef4a56e129b984a.json - resolving delete from //UE4/Release-4.22/EngineTest/Test/Screenshots/ActorMerging/ActorMerging_SimpleMeshMerge_LOD_2_None/Windows/D3D11_SM5/3fcc508dbd2d03ed4ef4a56e129b984a.json#2
Delete resolve:
at: delete
ay: ignore
/src/ROBOMERGE_RELEASE_Main/EngineTest/Test/Screenshots/ActorMerging/ActorMerging_SimpleMeshMerge_LOD_2_None/Windows/D3D11_SM5/3fcc508dbd2d03ed4ef4a56e129b984a.png - resolving delete from //UE4/Release-4.22/EngineTest/Test/Screenshots/ActorMerging/ActorMerging_SimpleMeshMerge_LOD_2_None/Windows/D3D11_SM5/3fcc508dbd2d03ed4ef4a56e129b984a.png#2
Delete resolve:
at: delete
ay: ignore
/src/ROBOMERGE_RELEASE_Main/EngineTest/Test/Screenshots/ActorMerging/ActorMerging_SimpleMeshMerge_LOD_3_None/Windows/D3D11_SM5/22c96594569807222f8f1a0ad282ed19.json - merging //UE4/Release-4.22/EngineTest/Test/Screenshots/ActorMerging/ActorMerging_SimpleMeshMerge_LOD_3_None/Windows/D3D11_SM5/22c96594569807222f8f1a0ad282ed19.json#3
/src/ROBOMERGE_RELEASE_Main/EngineTest/Test/Screenshots/ActorMerging/ActorMerging_SimpleMeshMerge_LOD_3_None/Windows/D3D11_SM5/22c96594569807222f8f1a0ad282ed19.png - vs //UE4/Release-4.22/EngineTest/Test/Screenshots/ActorMerging/ActorMerging_SimpleMeshMerge_LOD_3_None/Windows/D3D11_SM5/22c96594569807222f8f1a0ad282ed19.png#3
/src/ROBOMERGE_RELEASE_Main/EngineTest/Test/Screenshots/ActorMerging/ActorMerging_SimpleMeshMerge_LOD_3_None/Windows/D3D11_SM5/3fcc508dbd2d03ed4ef4a56e129b984a.json - resolving delete from //UE4/Release-4.22/EngineTest/Test/Screenshots/ActorMerging/ActorMerging_SimpleMeshMerge_LOD_3_None/Windows/D3D11_SM5/3fcc508dbd2d03ed4ef4a56e129b984a.json#2
Delete resolve:
at: delete
ay: ignore
/src/ROBOMERGE_RELEASE_Main/EngineTest/Test/Screenshots/ActorMerging/ActorMerging_SimpleMeshMerge_LOD_3_None/Windows/D3D11_SM5/3fcc508dbd2d03ed4ef4a56e129b984a.png - resolving delete from //UE4/Release-4.22/EngineTest/Test/Screenshots/ActorMerging/ActorMerging_SimpleMeshMerge_LOD_3_None/Windows/D3D11_SM5/3fcc508dbd2d03ed4ef4a56e129b984a.png#2
Delete resolve:
at: delete
ay: ignore

--------------------------------------
Revert a material fix done for UE-67819, this change broke old assets.

diff --git a/Engine/Source/Developer/MeshDescriptionOperations/Private/MeshDescriptionOperations.cpp b/Engine/Source/Developer/MeshDescriptionOperations/Private/MeshDescriptionOperations.cpp
--- a/Engine/Source/Developer/MeshDescriptionOperations/Private/MeshDescriptionOperations.cpp
+++ b/Engine/Source/Developer/MeshDescriptionOperations/Private/MeshDescriptionOperations.cpp
@@ -478,14 +478,15 @@ void FMeshDescriptionOperations::ConvertFromRawMesh(const FRawMesh& SourceRawMes
 	{
 		if (!MaterialIndexToPolygonGroup.Contains(MaterialIndex))
 		{
-			FPolygonGroupID PolygonGroupID = DestinationMeshDescription.CreatePolygonGroup();
+			FPolygonGroupID PolygonGroupID(MaterialIndex);
+			DestinationMeshDescription.CreatePolygonGroupWithID(PolygonGroupID);
 			if (MaterialMap.Contains(MaterialIndex))
 			{
 				PolygonGroupImportedMaterialSlotNames[PolygonGroupID] = MaterialMap[MaterialIndex];
 			}
 			else
 			{
-				PolygonGroupImportedMaterialSlotNames[PolygonGroupID] = FName(*FString::Printf(TEXT("MaterialSlot_%d"), PolygonGroupID.GetValue()));
+				PolygonGroupImportedMaterialSlotNames[PolygonGroupID] = FName(*FString::Printf(TEXT("MaterialSlot_%d"), MaterialIndex));
 			}
 			PolygonGroups.Add(PolygonGroupID);
 			MaterialIndexToPolygonGroup.Add(MaterialIndex, PolygonGroupID);
diff --git a/Engine/Source/Runtime/Engine/Private/StaticMesh.cpp b/Engine/Source/Runtime/Engine/Private/StaticMesh.cpp
--- a/Engine/Source/Runtime/Engine/Private/StaticMesh.cpp
+++ b/Engine/Source/Runtime/Engine/Private/StaticMesh.cpp
@@ -3659,7 +3659,7 @@ void UStaticMesh::FixupMaterialSlotName()
 // differences, etc.) replace the version GUID below with a new one.
 // In case of merge conflicts with DDC versions, you MUST generate a new GUID
 // and set this new GUID as the version.                                       
-#define MESHDATAKEY_STATICMESH_DERIVEDDATA_VER TEXT("F62CF84A92D948DEB06D56AEC759D2D3")
+#define MESHDATAKEY_STATICMESH_DERIVEDDATA_VER TEXT("6E342F4997294EBEA9FB3ED1E8CF5D3B")
 
 static const FString& GetMeshDataKeyStaticMeshDerivedDataVersion()
 {
