Fixed box mapping generation to fit in [0, 1] range (when there's no tiling).

diff --git a/Engine/Source/Developer/MeshDescriptionOperations/Private/MeshDescriptionOperations.cpp b/Engine/Source/Developer/MeshDescriptionOperations/Private/MeshDescriptionOperations.cpp
--- a/Engine/Source/Developer/MeshDescriptionOperations/Private/MeshDescriptionOperations.cpp
+++ b/Engine/Source/Developer/MeshDescriptionOperations/Private/MeshDescriptionOperations.cpp
@@ -1739,24 +1739,29 @@ void FMeshDescriptionOperations::GenerateBoxUV(const FMeshDescription& MeshDescr
 {
 	FVector Size = Params.Size * Params.Scale;
 	FVector HalfSize = Size / 2.0f;
-	FVector Offset = Params.Position - HalfSize;
-
-	FVector HintU = FVector::UpVector;
-	FVector HintV = FVector::RightVector;
 
 	TMeshAttributesConstRef<FVertexID, FVector> VertexPositions = MeshDescription.VertexAttributes().GetAttributesRef<FVector>(MeshAttribute::Vertex::Position);
 
 	OutTexCoords.AddZeroed(MeshDescription.VertexInstances().Num());
 
+	// Setup the UVs such that the mapping is from top-left to bottom-right when viewed orthographically
+	TArray<TPair<FVector, FVector>> PlaneUVs;
+	PlaneUVs.Add(TPair<FVector, FVector>(FVector::ForwardVector, FVector::RightVector));	// Top view
+	PlaneUVs.Add(TPair<FVector, FVector>(FVector::BackwardVector, FVector::RightVector));	// Bottom view
+	PlaneUVs.Add(TPair<FVector, FVector>(FVector::ForwardVector, FVector::DownVector));		// Right view
+	PlaneUVs.Add(TPair<FVector, FVector>(FVector::BackwardVector, FVector::DownVector));	// Left view
+	PlaneUVs.Add(TPair<FVector, FVector>(FVector::LeftVector, FVector::DownVector));		// Front view
+	PlaneUVs.Add(TPair<FVector, FVector>(FVector::RightVector, FVector::DownVector));		// Back view
+
 	TArray<FPlane> BoxPlanes;
 	const FVector& Center = Params.Position;
 
 	BoxPlanes.Add(FPlane(Center + FVector(0, 0, HalfSize.Z), FVector::UpVector));		// Top plane
-	BoxPlanes.Add(FPlane(Center - FVector(0, 0, HalfSize.Z), -FVector::UpVector));		// Bottom plane
-	BoxPlanes.Add(FPlane(Center + FVector(HalfSize.X, 0, 0), FVector::ForwardVector));	// Right plane
-	BoxPlanes.Add(FPlane(Center - FVector(HalfSize.X, 0, 0), -FVector::ForwardVector));	// Left plane
-	BoxPlanes.Add(FPlane(Center + FVector(0, HalfSize.Y, 0), FVector::RightVector));	// Front plane
-	BoxPlanes.Add(FPlane(Center - FVector(0, HalfSize.Y, 0), -FVector::RightVector));	// Back plane
+	BoxPlanes.Add(FPlane(Center - FVector(0, 0, HalfSize.Z), FVector::DownVector));		// Bottom plane
+	BoxPlanes.Add(FPlane(Center + FVector(0, HalfSize.Y, 0), FVector::RightVector));	// Right plane
+	BoxPlanes.Add(FPlane(Center - FVector(0, HalfSize.Y, 0), FVector::LeftVector));		// Left plane
+	BoxPlanes.Add(FPlane(Center + FVector(HalfSize.X, 0, 0), FVector::ForwardVector));	// Front plane
+	BoxPlanes.Add(FPlane(Center - FVector(HalfSize.X, 0, 0), FVector::BackwardVector));	// Back plane
 
 	// For each polygon, find the box plane that best matches the polygon normal
 	for (const FPolygonID& PolygonID : MeshDescription.Polygons().GetElementIDs())
@@ -1784,17 +1789,9 @@ void FMeshDescriptionOperations::GenerateBoxUV(const FMeshDescription& MeshDescr
 			}
 		}
 
-		const FPlane& BestPlane = BoxPlanes[BestPlaneIndex];
-
-		FVector U = HintU;
-		FVector V = BestPlane ^ HintU;
-
-		if (V.IsZero())
-		{
-			// Plane normal and U were aligned, so try with V instead
-			U = HintV;
-			V = BestPlane ^ HintV;
-		}
+		FVector U = PlaneUVs[BestPlaneIndex].Key;
+		FVector V = PlaneUVs[BestPlaneIndex].Value;
+		FVector Offset = Params.Position - HalfSize * (U + V);
 
 		for (const FVertexInstanceID& VertexInstanceID : VertexInstances)
 		{
