Fix the packing of UV, make sure the temporary destination has enough indices to receive the new computed UVLightMap

diff --git a/Engine/Source/Developer/MeshDescriptionOperations/Private/MeshDescriptionOperations.cpp b/Engine/Source/Developer/MeshDescriptionOperations/Private/MeshDescriptionOperations.cpp
--- a/Engine/Source/Developer/MeshDescriptionOperations/Private/MeshDescriptionOperations.cpp
+++ b/Engine/Source/Developer/MeshDescriptionOperations/Private/MeshDescriptionOperations.cpp
@@ -1446,6 +1446,15 @@ bool FMeshDescriptionOperations::GenerateUniqueUVsForStaticMesh(const FMeshDescr
 {
 	// Create a copy of original mesh (only copy necessary data)
 	FMeshDescription DuplicateMeshDescription(MeshDescription);
+	
+	//Make sure we have a destination UV TextureCoordinnate
+	{
+		TVertexInstanceAttributesRef<FVector2D> DuplicateVertexInstanceUVs = DuplicateMeshDescription.VertexInstanceAttributes().GetAttributesRef<FVector2D>(MeshAttribute::VertexInstance::TextureCoordinate);
+		if (DuplicateVertexInstanceUVs.GetNumIndices() < 2)
+		{
+			DuplicateVertexInstanceUVs.SetNumIndices(2);
+		}
+	}
 
 	TMap<FVertexInstanceID, FVertexInstanceID> RemapVertexInstance;
 	//Remove the identical material
@@ -1565,14 +1574,14 @@ bool FMeshDescriptionOperations::GenerateUniqueUVsForStaticMesh(const FMeshDescr
 		TVertexInstanceAttributesConstRef<FVector2D> DupVertexInstanceUVs = DuplicateMeshDescription.VertexInstanceAttributes().GetAttributesRef<FVector2D>(MeshAttribute::VertexInstance::TextureCoordinate);
 		TVertexInstanceAttributesConstRef<FVector2D> VertexInstanceUVs = MeshDescription.VertexInstanceAttributes().GetAttributesRef<FVector2D>(MeshAttribute::VertexInstance::TextureCoordinate);
 		// Save generated UVs
-		check(VertexInstanceUVs.GetNumIndices() > 1);
+		check(DupVertexInstanceUVs.GetNumIndices() > 1);
 		OutTexCoords.AddZeroed(VertexInstanceUVs.GetNumElements());
 		int32 TextureCoordIndex = 0;
 		for (const FVertexInstanceID& VertexInstanceID : MeshDescription.VertexInstances().GetElementIDs())
 		{
 			FVertexInstanceID RemapID = bMergeIdenticalMaterials ? RemapVertexInstance[VertexInstanceID] : VertexInstanceID;
 			// Save generated UVs
-			OutTexCoords[TextureCoordIndex] = VertexInstanceUVs.Get(RemapID, 1);	// UV1
+			OutTexCoords[TextureCoordIndex] = DupVertexInstanceUVs.Get(RemapID, 1);	// UV1
 			TextureCoordIndex++;
 		}
 	}
