Fix MinLightmapResolution not considering 2 texel padding in FMeshDescriptionOperations

diff --git a/Engine/Source/Developer/MeshDescriptionOperations/Private/MeshDescriptionOperations.cpp b/Engine/Source/Developer/MeshDescriptionOperations/Private/MeshDescriptionOperations.cpp
--- a/Engine/Source/Developer/MeshDescriptionOperations/Private/MeshDescriptionOperations.cpp
+++ b/Engine/Source/Developer/MeshDescriptionOperations/Private/MeshDescriptionOperations.cpp
@@ -1433,6 +1433,11 @@ bool FMeshDescriptionOperations::CreateLightMapUVLayout(FMeshDescription& MeshDe
 	FLayoutUV Packer(MeshDescriptionView);
 	Packer.SetVersion(LightmapUVVersion);
 
+	if (LightmapUVVersion >= ELightmapUVVersion::ForceLightmapPadding)
+	{
+		MinLightmapResolution -= 2;
+	}
+
 	Packer.FindCharts(OverlappingCorners);
 	bool bPackSuccess = Packer.FindBestPacking(MinLightmapResolution);
 	if (bPackSuccess)
diff --git a/Engine/Source/Developer/MeshUtilitiesCommon/Public/MeshUtilitiesCommon.h b/Engine/Source/Developer/MeshUtilitiesCommon/Public/MeshUtilitiesCommon.h
--- a/Engine/Source/Developer/MeshUtilitiesCommon/Public/MeshUtilitiesCommon.h
+++ b/Engine/Source/Developer/MeshUtilitiesCommon/Public/MeshUtilitiesCommon.h
@@ -13,7 +13,8 @@ enum class ELightmapUVVersion : int32
 	ChartJoiningLFix = 4,
 	Allocator2DFlipFix = 5,
 	ConsiderLightmapPadding = 6,
-	Latest = ConsiderLightmapPadding
+	ForceLightmapPadding = 7,
+	Latest = ForceLightmapPadding
 };
 
 /** Helper struct for building acceleration structures. */
