Integrating CL# 9607649 from Dev-Enterprise
Adding nan check when generating normals + fixed reimporting SceneFbx could corrupt the default import values config

diff --git a/Engine/Source/Developer/MeshDescriptionOperations/Private/MeshDescriptionOperations.cpp b/Engine/Source/Developer/MeshDescriptionOperations/Private/MeshDescriptionOperations.cpp
--- a/Engine/Source/Developer/MeshDescriptionOperations/Private/MeshDescriptionOperations.cpp
+++ b/Engine/Source/Developer/MeshDescriptionOperations/Private/MeshDescriptionOperations.cpp
@@ -892,16 +892,26 @@ void FMeshDescriptionOperations::CreatePolygonNTB(FMeshDescription& MeshDescript
 						// Use InverseSlow to catch singular matrices.  Inverse can miss this sometimes.
 						const FMatrix TextureToLocal = ParameterToTexture.Inverse() * ParameterToLocal;
 
-						FVector TmpTangentX(0.0f);
-						FVector TmpTangentY(0.0f);
-						FVector TmpTangentZ(0.0f);
-						TmpTangentX = TextureToLocal.TransformVector(FVector(1, 0, 0)).GetSafeNormal();
-						TmpTangentY = TextureToLocal.TransformVector(FVector(0, 1, 0)).GetSafeNormal();
-						TmpTangentZ = Normal;
+						FVector TmpTangentX = TextureToLocal.TransformVector(FVector(1, 0, 0)).GetSafeNormal();
+						FVector TmpTangentY = TextureToLocal.TransformVector(FVector(0, 1, 0)).GetSafeNormal();
+						FVector TmpTangentZ = Normal;
 						FVector::CreateOrthonormalBasis(TmpTangentX, TmpTangentY, TmpTangentZ);
+
+						if (TmpTangentX.IsNearlyZero() || TmpTangentX.ContainsNaN()
+							|| TmpTangentY.IsNearlyZero() || TmpTangentY.ContainsNaN())
+						{
+							TmpTangentX = FVector::ZeroVector;
+							TmpTangentY = FVector::ZeroVector;
+						}
+
+						if (TmpTangentZ.IsNearlyZero() || TmpTangentZ.ContainsNaN())
+						{
+							TmpTangentZ = FVector::ZeroVector;
+						}
+
 						TangentX += TmpTangentX;
 						TangentY += TmpTangentY;
-						TangentZ += TmpTangentZ;
+						TangentZ += TmpTangentZ;						
 					}
 					else
 					{
diff --git a/Engine/Source/Developer/MeshUtilities/Private/MeshUtilities.cpp b/Engine/Source/Developer/MeshUtilities/Private/MeshUtilities.cpp
--- a/Engine/Source/Developer/MeshUtilities/Private/MeshUtilities.cpp
+++ b/Engine/Source/Developer/MeshUtilities/Private/MeshUtilities.cpp
@@ -1112,12 +1112,16 @@ static void ComputeTriangleTangents(
 				OutTangentY[TriangleIndex],
 				OutTangentZ[TriangleIndex]
 			);
+
 			if (OutTangentX[TriangleIndex].IsNearlyZero() || OutTangentX[TriangleIndex].ContainsNaN()
-				|| OutTangentY[TriangleIndex].IsNearlyZero() || OutTangentY[TriangleIndex].ContainsNaN()
-				|| OutTangentZ[TriangleIndex].IsNearlyZero() || OutTangentZ[TriangleIndex].ContainsNaN())
+				|| OutTangentY[TriangleIndex].IsNearlyZero() || OutTangentY[TriangleIndex].ContainsNaN())
 			{
 				OutTangentX[TriangleIndex] = FVector::ZeroVector;
 				OutTangentY[TriangleIndex] = FVector::ZeroVector;
+			}
+
+			if (OutTangentZ[TriangleIndex].IsNearlyZero() || OutTangentZ[TriangleIndex].ContainsNaN())
+			{
 				OutTangentZ[TriangleIndex] = FVector::ZeroVector;
 			}
 		}
diff --git a/Engine/Source/Editor/UnrealEd/Private/Fbx/ReimportFbxSceneFactory.cpp b/Engine/Source/Editor/UnrealEd/Private/Fbx/ReimportFbxSceneFactory.cpp
--- a/Engine/Source/Editor/UnrealEd/Private/Fbx/ReimportFbxSceneFactory.cpp
+++ b/Engine/Source/Editor/UnrealEd/Private/Fbx/ReimportFbxSceneFactory.cpp
@@ -609,11 +609,11 @@ EReimportResult::Type UReimportFbxSceneFactory::Reimport(UObject* Obj)
 	SFbxSceneOptionWindow::CopyFbxOptionsToFbxOptions(GlobalImportSettingsReference, GlobalImportSettings);
 	SFbxSceneOptionWindow::CopyFbxOptionsToStaticMeshOptions(GlobalImportSettingsReference, SceneImportOptionsStaticMesh);
 	SceneImportOptionsStaticMesh->FillStaticMeshInmportData(StaticMeshImportData, SceneImportOptions);
-	StaticMeshImportData->SaveConfig();
+	SceneImportOptionsStaticMesh->SaveConfig();
 
 	SFbxSceneOptionWindow::CopyFbxOptionsToSkeletalMeshOptions(GlobalImportSettingsReference, SceneImportOptionsSkeletalMesh);
 	SceneImportOptionsSkeletalMesh->FillSkeletalMeshInmportData(SkeletalMeshImportData, AnimSequenceImportData, SceneImportOptions);
-	SkeletalMeshImportData->SaveConfig();
+	SceneImportOptionsSkeletalMesh->SaveConfig();
 
 	//Update the blueprint
 	UBlueprint *ReimportBlueprint = nullptr;
