Compare all set of UVs when trying to merge polygons in FMeshDescriptionOperations::GenerateUniqueUVsForStaticMesh
Swap vertices in material baking when mesh is mirrored

diff --git a/Engine/Source/Developer/MaterialBaking/Private/MaterialRenderItem.cpp b/Engine/Source/Developer/MaterialBaking/Private/MaterialRenderItem.cpp
--- a/Engine/Source/Developer/MaterialBaking/Private/MaterialRenderItem.cpp
+++ b/Engine/Source/Developer/MaterialBaking/Private/MaterialRenderItem.cpp
@@ -256,6 +256,12 @@ void FMeshMaterialRenderItem::PopulateWithMeshData()
 				Indices.Add(VertIndex - 3);
 				Indices.Add(VertIndex - 1);
 				Indices.Add(VertIndex - 2);
+
+				// Swap vertices order if mesh is mirrored
+				if (MeshSettings->bMirrored)
+				{
+					Swap(Vertices[Vertices.Num() - 3].Position, Vertices[Vertices.Num() - 1].Position);
+				}
 			}
 			FaceIndex++;
 		}
diff --git a/Engine/Source/Developer/MeshDescriptionOperations/Private/MeshDescriptionOperations.cpp b/Engine/Source/Developer/MeshDescriptionOperations/Private/MeshDescriptionOperations.cpp
--- a/Engine/Source/Developer/MeshDescriptionOperations/Private/MeshDescriptionOperations.cpp
+++ b/Engine/Source/Developer/MeshDescriptionOperations/Private/MeshDescriptionOperations.cpp
@@ -1611,15 +1611,22 @@ bool FMeshDescriptionOperations::GenerateUniqueUVsForStaticMesh(const FMeshDescr
 		TArray<FPolygonID> ToDeletePolygons;
 		RemapVertexInstance.Reserve(DuplicateMeshDescription.VertexInstances().Num());
 		TArray<FPolygonID> UniquePolygons;
+
+		TArray<FVector2D> RefUVs;
 		for (FPolygonID RefPolygonID : DuplicateMeshDescription.Polygons().GetElementIDs())
 		{
 			FPolygonGroupID RefPolygonGroupID = DuplicateMeshDescription.GetPolygonPolygonGroup(RefPolygonID);
-			const TArray<FVertexInstanceID>& RefVertexInstances = DuplicateMeshDescription.GetPolygonVertexInstances(RefPolygonID);
-			TArray<FVector2D> RefUVs;
+			const TArray<FVertexInstanceID>& RefVertexInstances = DuplicateMeshDescription.GetPolygonPerimeterVertexInstances(RefPolygonID);
+
+			RefUVs.Empty(RefVertexInstances.Num() * VertexInstanceUVs.GetNumIndices());
 			for (FVertexInstanceID RefVertexInstanceID : RefVertexInstances)
 			{
-				RefUVs.Add(VertexInstanceUVs[RefVertexInstanceID]);
+				for (int32 UVChannel = 0; UVChannel < VertexInstanceUVs.GetNumIndices(); ++UVChannel)
+				{
+					RefUVs.Add(VertexInstanceUVs.Get(RefVertexInstanceID, UVChannel));
+				}
 			}
+
 			FPolygonID MatchPolygonID = FPolygonID::Invalid;
 			for (FPolygonID TestPolygonID : UniquePolygons)
 			{
@@ -1628,21 +1635,27 @@ bool FMeshDescriptionOperations::GenerateUniqueUVsForStaticMesh(const FMeshDescr
 				{
 					continue;
 				}
-				const TArray<FVertexInstanceID>& TestVertexInstances = DuplicateMeshDescription.GetPolygonVertexInstances(TestPolygonID);
+
+				const TArray<FVertexInstanceID>& TestVertexInstances = DuplicateMeshDescription.GetPolygonPerimeterVertexInstances(TestPolygonID);
 				if (TestVertexInstances.Num() != RefVertexInstances.Num())
 				{
 					continue;
 				}
+				
 				bool bIdentical = true;
 				int32 UVIndex = 0;
 				for (FVertexInstanceID TestVertexInstanceID : TestVertexInstances)
 				{
-					if (VertexInstanceUVs[TestVertexInstanceID] != RefUVs[UVIndex])
+					// All UV channels must match for polygons to be identical
+					for (int32 UVChannel = 0; bIdentical && UVChannel < VertexInstanceUVs.GetNumIndices(); ++UVIndex, ++UVChannel)
+					{
+						bIdentical = VertexInstanceUVs.Get(TestVertexInstanceID, UVChannel) == RefUVs[UVIndex];
+					}
+
+					if (!bIdentical)
 					{
-						bIdentical = false;
 						break;
 					}
-					UVIndex++;
 				}
 				if (bIdentical)
 				{
@@ -1661,7 +1674,7 @@ bool FMeshDescriptionOperations::GenerateUniqueUVsForStaticMesh(const FMeshDescr
 			}
 			else
 			{
-				const TArray<FVertexInstanceID>& TestVertexInstances = DuplicateMeshDescription.GetPolygonVertexInstances(MatchPolygonID);
+				const TArray<FVertexInstanceID>& TestVertexInstances = DuplicateMeshDescription.GetPolygonPerimeterVertexInstances(MatchPolygonID);
 				int32 VertexInstanceIndex = 0;
 				for (FVertexInstanceID RefVertexInstanceID : RefVertexInstances)
 				{
