Fixed CIS warning
#rnx

diff --git a/Engine/Source/Developer/MeshDescriptionOperations/Private/MeshDescriptionOperations.cpp b/Engine/Source/Developer/MeshDescriptionOperations/Private/MeshDescriptionOperations.cpp
--- a/Engine/Source/Developer/MeshDescriptionOperations/Private/MeshDescriptionOperations.cpp
+++ b/Engine/Source/Developer/MeshDescriptionOperations/Private/MeshDescriptionOperations.cpp
@@ -1616,7 +1616,7 @@ bool FMeshDescriptionOperations::GenerateUniqueUVsForStaticMesh(const FMeshDescr
 		for (FPolygonID RefPolygonID : DuplicateMeshDescription.Polygons().GetElementIDs())
 		{
 			FPolygonGroupID RefPolygonGroupID = DuplicateMeshDescription.GetPolygonPolygonGroup(RefPolygonID);
-			const TArray<FVertexInstanceID>& RefVertexInstances = DuplicateMeshDescription.GetPolygonPerimeterVertexInstances(RefPolygonID);
+			const TArray<FVertexInstanceID>& RefVertexInstances = DuplicateMeshDescription.GetPolygonVertexInstances(RefPolygonID);
 
 			RefUVs.Empty(RefVertexInstances.Num() * VertexInstanceUVs.GetNumIndices());
 			for (FVertexInstanceID RefVertexInstanceID : RefVertexInstances)
@@ -1636,7 +1636,7 @@ bool FMeshDescriptionOperations::GenerateUniqueUVsForStaticMesh(const FMeshDescr
 					continue;
 				}
 
-				const TArray<FVertexInstanceID>& TestVertexInstances = DuplicateMeshDescription.GetPolygonPerimeterVertexInstances(TestPolygonID);
+				const TArray<FVertexInstanceID>& TestVertexInstances = DuplicateMeshDescription.GetPolygonVertexInstances(TestPolygonID);
 				if (TestVertexInstances.Num() != RefVertexInstances.Num())
 				{
 					continue;
@@ -1674,7 +1674,7 @@ bool FMeshDescriptionOperations::GenerateUniqueUVsForStaticMesh(const FMeshDescr
 			}
 			else
 			{
-				const TArray<FVertexInstanceID>& TestVertexInstances = DuplicateMeshDescription.GetPolygonPerimeterVertexInstances(MatchPolygonID);
+				const TArray<FVertexInstanceID>& TestVertexInstances = DuplicateMeshDescription.GetPolygonVertexInstances(MatchPolygonID);
 				int32 VertexInstanceIndex = 0;
 				for (FVertexInstanceID RefVertexInstanceID : RefVertexInstances)
 				{
