MeshDescription now uses triangles as its primary face representation. They can and should be used in preference to polygons when a mesh is known to be entirely triangulated, as they are quicker, and 

diff --git a/Engine/Plugins/Runtime/EditableMesh/Source/EditableMesh/Public/EditableMeshTypes.h b/Engine/Plugins/Runtime/EditableMesh/Source/EditableMesh/Public/EditableMeshTypes.h
--- a/Engine/Plugins/Runtime/EditableMesh/Source/EditableMesh/Public/EditableMeshTypes.h
+++ b/Engine/Plugins/Runtime/EditableMesh/Source/EditableMesh/Public/EditableMeshTypes.h
@@ -608,9 +608,9 @@ struct FEdgeToCreate
 	UPROPERTY( BlueprintReadWrite, Category="Editable Mesh" )
 	FVertexID VertexID1;
 
-	/** The polygons that are connected to this edge. */
-	UPROPERTY( BlueprintReadWrite, Category="Editable Mesh" )
-	TArray<FPolygonID> ConnectedPolygons;
+	///** The polygons that are connected to this edge. */
+	//UPROPERTY( BlueprintReadWrite, Category="Editable Mesh" )
+	//TArray<FPolygonID> ConnectedPolygons;
 
 	/** Attributes of this edge itself */
 	UPROPERTY( BlueprintReadWrite, Category="Editable Mesh" )
@@ -624,7 +624,7 @@ struct FEdgeToCreate
 	FEdgeToCreate()
 		: VertexID0( FVertexID::Invalid ),
 		  VertexID1( FVertexID::Invalid ),
-		  ConnectedPolygons(),
+//		  ConnectedPolygons(),
 		  EdgeAttributes(),
 		  OriginalEdgeID( FEdgeID::Invalid )
 	{
@@ -633,12 +633,12 @@ struct FEdgeToCreate
 	FString ToString() const
 	{
 		return FString::Printf(
-			TEXT( "VertexID0:%s, VertexID1:%s, ConnectedPolygons:%s, OriginalEdgeID:%s" ),
+			TEXT( "VertexID0:%s, VertexID1:%s, EdgeAttributes:%s, OriginalEdgeID:%s" ),
 			*VertexID0.ToString(), 
 			*VertexID1.ToString(),
-			*LogHelpers::ArrayToString( ConnectedPolygons ),
+//			*LogHelpers::ArrayToString( ConnectedPolygons ),
 			*EdgeAttributes.ToString(),
-			OriginalEdgeID.GetValue() );
+			*OriginalEdgeID.ToString() );
 	}
 };
 
@@ -1035,6 +1035,7 @@ struct FChangeVertexInstancesForPolygon
 };
 
 
+#if 0
 USTRUCT( BlueprintType )
 struct FVerticesForEdge
 {
@@ -1069,6 +1070,7 @@ struct FVerticesForEdge
 			*NewVertexID1.ToString() );
 	}
 };
+#endif
 
 
 USTRUCT( BlueprintType )
diff --git a/Engine/Plugins/Runtime/ProceduralMeshComponent/Source/ProceduralMeshComponentEditor/Private/ProceduralMeshComponentDetails.cpp b/Engine/Plugins/Runtime/ProceduralMeshComponent/Source/ProceduralMeshComponentEditor/Private/ProceduralMeshComponentDetails.cpp
--- a/Engine/Plugins/Runtime/ProceduralMeshComponent/Source/ProceduralMeshComponentEditor/Private/ProceduralMeshComponentDetails.cpp
+++ b/Engine/Plugins/Runtime/ProceduralMeshComponent/Source/ProceduralMeshComponentEditor/Private/ProceduralMeshComponentDetails.cpp
@@ -223,10 +223,7 @@ FReply FProceduralMeshComponentDetails::ClickedOnConvertToStaticMesh()
 					}
 
 					// Insert a polygon into the mesh
-					const FPolygonID NewPolygonID = MeshDescription.CreatePolygon(PolygonGroupID, VertexInstanceIDs);
-					//Triangulate the polygon
-					FMeshPolygon& Polygon = MeshDescription.GetPolygon(NewPolygonID);
-					MeshDescription.ComputePolygonTriangulation(NewPolygonID, Polygon.Triangles);
+					MeshDescription.CreatePolygon(PolygonGroupID, VertexInstanceIDs);
 				}
 			}
 
diff --git a/Engine/Plugins/Tests/EditorTests/Source/EditorTests/Private/UnrealEd/MeshDescriptionTests.cpp b/Engine/Plugins/Tests/EditorTests/Source/EditorTests/Private/UnrealEd/MeshDescriptionTests.cpp
--- a/Engine/Plugins/Tests/EditorTests/Source/EditorTests/Private/UnrealEd/MeshDescriptionTests.cpp
+++ b/Engine/Plugins/Tests/EditorTests/Source/EditorTests/Private/UnrealEd/MeshDescriptionTests.cpp
@@ -614,8 +614,8 @@ bool FMeshDescriptionTest::NTBTest(FAutomationTestExecutionInfo& ExecutionInfo)
 			}
 			const FPolygonGroupID& PolygonGroupID = MeshDescription.GetPolygonPolygonGroup(PolygonID);
 			int32 PolygonIDValue = PolygonID.GetValue();
-			const TArray<FMeshTriangle>& Triangles = MeshDescription.GetPolygonTriangles(PolygonID);
-			for (const FMeshTriangle& MeshTriangle : Triangles)
+			const TArray<FTriangleID>& TriangleIDs = MeshDescription.GetPolygonTriangleIDs(PolygonID);
+			for (const FTriangleID TriangleID : TriangleIDs)
 			{
 				if (bError)
 				{
@@ -624,7 +624,7 @@ bool FMeshDescriptionTest::NTBTest(FAutomationTestExecutionInfo& ExecutionInfo)
 				for (int32 Corner = 0; Corner < 3 && bError == false; ++Corner)
 				{
 					uint32 WedgeIndex = TriangleIndex * 3 + Corner;
-					const FVertexInstanceID VertexInstanceID = MeshTriangle.GetVertexInstanceID(Corner);
+					const FVertexInstanceID VertexInstanceID = MeshDescription.GetTriangleVertexInstance(TriangleID, Corner);
 					const int32 VertexInstanceIDValue = VertexInstanceID.GetValue();
 					if (RawMesh.WedgeColors[WedgeIndex] != FLinearColor(VertexInstanceColors[VertexInstanceID]).ToFColor(true))
 					{
@@ -688,4 +688,258 @@ bool FMeshDescriptionTest::NTBTest(FAutomationTestExecutionInfo& ExecutionInfo)
 	}
 
 	return bAllSame;
-}
\ No newline at end of file
+}
+
+
+IMPLEMENT_SIMPLE_AUTOMATION_TEST(FMeshDescriptionBuilderTest, "Editor.Meshes.MeshDescription.Builder", EAutomationTestFlags::EditorContext | EAutomationTestFlags::EngineFilter);
+
+bool FMeshDescriptionBuilderTest::RunTest(const FString& Parameters)
+{
+	FMeshDescription MeshDescription;
+	MeshDescription.VertexAttributes().RegisterAttribute<FVector>(MeshAttribute::Vertex::Position);
+	TVertexAttributesRef<FVector> Positions = MeshDescription.VertexAttributes().GetAttributesRef<FVector>(MeshAttribute::Vertex::Position);
+
+	// Build a hexagonal cylinder
+
+	const float Radius = 100.0f;
+	const float Height = 200.0f;
+	const int32 NumSides = 6;
+
+	FVertexID TopVertexIDs[NumSides];
+	FVertexID BottomVertexIDs[NumSides];
+
+	FVertexInstanceID TopVertexInstanceIDs[NumSides];
+	FVertexInstanceID SideTopVertexInstanceIDs[NumSides];
+	FVertexInstanceID BottomVertexInstanceIDs[NumSides];
+	FVertexInstanceID SideBottomVertexInstanceIDs[NumSides];
+
+	FEdgeID SideEdgeIDs[NumSides];
+	FPolygonID SidePolygonIDs[NumSides];
+	FTriangleID TopTriangleIDs[NumSides];
+	FTriangleID BottomTriangleIDs[NumSides];
+
+	for (int32 Index = 0; Index < NumSides; ++Index)
+	{
+		// Create vertices at top and bottom
+
+		const float X = Radius * FMath::Cos(PI * 2.0f * Index / NumSides);
+		const float Y = Radius * FMath::Sin(PI * 2.0f * Index / NumSides);
+
+		TopVertexIDs[Index] = MeshDescription.CreateVertex();
+		Positions[TopVertexIDs[Index]] = FVector(X, Y, -Height);
+
+		BottomVertexIDs[Index] = MeshDescription.CreateVertex();
+		Positions[BottomVertexIDs[Index]] = FVector(X, Y, Height);
+
+		// Each vertex has two vertex instances: one for the side polygons, and one for the top/bottom polygons
+
+		TopVertexInstanceIDs[Index] = MeshDescription.CreateVertexInstance(TopVertexIDs[Index]);
+		SideTopVertexInstanceIDs[Index] = MeshDescription.CreateVertexInstance(TopVertexIDs[Index]);
+
+		BottomVertexInstanceIDs[Index] = MeshDescription.CreateVertexInstance(BottomVertexIDs[Index]);
+		SideBottomVertexInstanceIDs[Index] = MeshDescription.CreateVertexInstance(BottomVertexIDs[Index]);
+	}
+
+	// Create central vertex at bottom, and associated instances
+
+	FVertexID BottomCenterVertexID = MeshDescription.CreateVertex();
+	Positions[BottomCenterVertexID] = FVector(0.0f, 0.0f, Height);
+
+	FVertexInstanceID BottomCenterVertexInstanceID = MeshDescription.CreateVertexInstance(BottomCenterVertexID);
+
+	// Test vertex / vertex instance status
+
+	TestEqual("Vertex count", MeshDescription.Vertices().Num(), NumSides * 2 + 1);
+	TestEqual("Vertex array size", MeshDescription.Vertices().GetArraySize(), NumSides * 2 + 1);
+	TestEqual("Vertex position count", Positions.GetNumElements(), NumSides * 2 + 1);
+
+	TestEqual("Vertex instance count", MeshDescription.VertexInstances().Num(), NumSides * 4 + 1);
+	TestEqual("Vertex instance array size", MeshDescription.VertexInstances().GetArraySize(), NumSides * 4 + 1);
+
+	for (int32 Index = 0; Index < NumSides; ++Index)
+	{
+		TestEqual("Vertex instances per vertex", MeshDescription.GetNumVertexVertexInstances(TopVertexIDs[Index]), 2);
+		TestTrue("Vertex contains vertex instance", MeshDescription.GetVertexVertexInstances(TopVertexIDs[Index]).Contains(TopVertexInstanceIDs[Index]));
+		TestTrue("Vertex contains vertex instance", MeshDescription.GetVertexVertexInstances(TopVertexIDs[Index]).Contains(SideTopVertexInstanceIDs[Index]));
+		TestEqual("Vertex instances per vertex", MeshDescription.GetNumVertexVertexInstances(BottomVertexIDs[Index]), 2);
+		TestTrue("Vertex contains vertex instance", MeshDescription.GetVertexVertexInstances(BottomVertexIDs[Index]).Contains(BottomVertexInstanceIDs[Index]));
+		TestTrue("Vertex contains vertex instance", MeshDescription.GetVertexVertexInstances(BottomVertexIDs[Index]).Contains(SideBottomVertexInstanceIDs[Index]));
+		TestEqual("Vertex instance parent", MeshDescription.GetVertexInstanceVertex(TopVertexInstanceIDs[Index]), TopVertexIDs[Index]);
+		TestEqual("Vertex instance parent", MeshDescription.GetVertexInstanceVertex(SideTopVertexInstanceIDs[Index]), TopVertexIDs[Index]);
+		TestEqual("Vertex instance parent", MeshDescription.GetVertexInstanceVertex(BottomVertexInstanceIDs[Index]), BottomVertexIDs[Index]);
+		TestEqual("Vertex instance parent", MeshDescription.GetVertexInstanceVertex(SideBottomVertexInstanceIDs[Index]), BottomVertexIDs[Index]);
+	}
+
+	// Add edges for side polys
+
+	for (int32 Index = 0; Index < NumSides; ++Index)
+	{
+		SideEdgeIDs[Index] = MeshDescription.CreateEdge(BottomVertexIDs[Index], TopVertexIDs[Index]);
+
+		TestTrue("Vertex connected edges", MeshDescription.GetVertexConnectedEdges(TopVertexIDs[Index]).Contains(SideEdgeIDs[Index]));
+		TestTrue("Vertex connected edges", MeshDescription.GetVertexConnectedEdges(BottomVertexIDs[Index]).Contains(SideEdgeIDs[Index]));
+	}
+
+	// Add polygon groups
+
+	FPolygonGroupID SidePolygonGroupID = MeshDescription.CreatePolygonGroup();
+	FPolygonGroupID TopPolygonGroupID = MeshDescription.CreatePolygonGroup();
+	FPolygonGroupID BottomPolygonGroupID = MeshDescription.CreatePolygonGroup();
+
+	// Now add side polygons
+
+	for (int32 Index = 0; Index < NumSides; ++Index)
+	{
+		// Define a side polygon to create
+		// We have already explicitly created the side edges, so we expect CreatePolygon to add two extra ones, top and bottom
+
+		const FVertexInstanceID SidePolyVertexInstanceIDs[4] =
+		{
+			SideBottomVertexInstanceIDs[Index],
+			SideBottomVertexInstanceIDs[(Index + 1) % NumSides],
+			SideTopVertexInstanceIDs[(Index + 1) % NumSides],
+			SideTopVertexInstanceIDs[Index]
+		};
+
+		TArray<FEdgeID> OutEdgeIDs;
+		SidePolygonIDs[Index] = MeshDescription.CreatePolygon(SidePolygonGroupID, SidePolyVertexInstanceIDs, &OutEdgeIDs);
+
+		TestEqual("Extra edges created", OutEdgeIDs.Num(), 2);
+		TestEqual("Bottom edge vertex 0", MeshDescription.GetEdgeVertex(OutEdgeIDs[0], 0), BottomVertexIDs[Index]);
+		TestEqual("Bottom edge vertex 1", MeshDescription.GetEdgeVertex(OutEdgeIDs[0], 1), BottomVertexIDs[(Index + 1) % NumSides]);
+		TestEqual("Top edge vertex 0", MeshDescription.GetEdgeVertex(OutEdgeIDs[1], 0), TopVertexIDs[(Index + 1) % NumSides]);
+		TestEqual("Top edge vertex 1", MeshDescription.GetEdgeVertex(OutEdgeIDs[1], 1), TopVertexIDs[Index]);
+
+		TestEqual("Triangles per side poly", MeshDescription.GetNumPolygonTriangles(SidePolygonIDs[Index]), 2);
+
+		TArray<FEdgeID, TInlineAllocator<1>> InternalEdgeIDs = MeshDescription.GetPolygonInternalEdges<TInlineAllocator<1>>(SidePolygonIDs[Index]);
+		TestEqual("Number of internal edges", InternalEdgeIDs.Num(), 1);
+	}
+
+	// Top face is a single polygon
+
+	TArray<FEdgeID> OutTopPolyEdgeIDs;
+	FPolygonID TopPolygonID = MeshDescription.CreatePolygon(TopPolygonGroupID, TopVertexInstanceIDs, &OutTopPolyEdgeIDs);
+
+	TestEqual("Extra edges created", OutTopPolyEdgeIDs.Num(), 0);
+	
+	TArray<FEdgeID, TInlineAllocator<NumSides - 3>> TopInternalEdgeIDs = MeshDescription.GetPolygonInternalEdges<TInlineAllocator<NumSides - 3>>(TopPolygonID);
+	TestEqual("Number of internal edges", TopInternalEdgeIDs.Num(), NumSides - 3);
+
+	// Bottom face is composed of individual triangles
+
+	for (int32 Index = 0; Index < NumSides; ++Index)
+	{
+		const FVertexInstanceID BottomTriangle[3] =
+		{
+			BottomVertexInstanceIDs[(Index + 1) % NumSides],
+			BottomVertexInstanceIDs[Index],
+			BottomCenterVertexInstanceID
+		};
+
+		TArray<FEdgeID> OutBottomTriEdgeIDs;
+		BottomTriangleIDs[Index] = MeshDescription.CreateTriangle(BottomPolygonGroupID, BottomTriangle, &OutBottomTriEdgeIDs);
+
+		const int32 ExtraEdges = (Index == 0) ? 2 : (Index == NumSides - 1) ? 0 : 1;
+		TestEqual("Extra bottom triangle edges", OutBottomTriEdgeIDs.Num(), ExtraEdges);
+	}
+
+	// Check totals correct
+
+	TestEqual("Total triangle count", MeshDescription.Triangles().Num(), NumSides * 4 - 2);
+	TestEqual("Total edge count", MeshDescription.Edges().Num(), NumSides * 6 - 3);
+	TestEqual("Polygons in side group", MeshDescription.GetNumPolygonGroupPolygons(SidePolygonGroupID), NumSides);
+	TestEqual("Polygons in top group", MeshDescription.GetNumPolygonGroupPolygons(TopPolygonGroupID), 1);
+	TestEqual("Polygons in bottom group", MeshDescription.GetNumPolygonGroupPolygons(BottomPolygonGroupID), NumSides);
+
+	// Check connections configured correctly
+
+	for (int32 Index = 0; Index < NumSides; ++Index)
+	{
+		int32 PrevIndex = (Index + NumSides - 1) % NumSides;
+
+		TestEqual("Side polygons connected to edge", MeshDescription.GetNumEdgeConnectedPolygons(SideEdgeIDs[Index]), 2);
+		TestEqual("Side triangles connected to edge", MeshDescription.GetNumEdgeConnectedTriangles(SideEdgeIDs[Index]), 2);
+		TestTrue("Check polygon connected to edge", MeshDescription.GetEdgeConnectedPolygons(SideEdgeIDs[Index]).Contains(SidePolygonIDs[Index]));
+		TestTrue("Check polygon connected to edge", MeshDescription.GetEdgeConnectedPolygons(SideEdgeIDs[Index]).Contains(SidePolygonIDs[PrevIndex]));
+
+		TestEqual("Side polygons connected to vertex instance", MeshDescription.GetNumVertexInstanceConnectedPolygons(SideTopVertexInstanceIDs[Index]), 2);
+		TestEqual("Side triangles connected to vertex instance", MeshDescription.GetNumVertexInstanceConnectedTriangles(SideTopVertexInstanceIDs[Index]), 3);
+		TestTrue("Check polygon connected to vertex instance", MeshDescription.GetVertexInstanceConnectedPolygons(SideTopVertexInstanceIDs[Index]).Contains(SidePolygonIDs[Index]));
+		TestTrue("Check polygon connected to vertex instance", MeshDescription.GetVertexInstanceConnectedPolygons(SideTopVertexInstanceIDs[Index]).Contains(SidePolygonIDs[PrevIndex]));
+	}
+
+	// Check adjacency
+
+	TArray<FPolygonID> AdjacentPolygonIDs = MeshDescription.GetPolygonAdjacentPolygons(SidePolygonIDs[0]);
+	TestEqual("Adjacent polygon count", AdjacentPolygonIDs.Num(), 4);
+	TestTrue("Check adjacency", AdjacentPolygonIDs.Contains(SidePolygonIDs[1]));
+	TestTrue("Check adjacency", AdjacentPolygonIDs.Contains(SidePolygonIDs[NumSides - 1]));
+	TestTrue("Check adjacency", AdjacentPolygonIDs.Contains(MeshDescription.GetTrianglePolygon(BottomTriangleIDs[0])));
+	TestTrue("Check adjacency", AdjacentPolygonIDs.Contains(TopPolygonID));
+
+	// Move some points and check retriangulation
+
+	Positions[TopVertexIDs[0]] = FVector(Radius * 0.25f, 0.0f, -Height);	// create a concave top and bottom
+	Positions[BottomVertexIDs[0]] = FVector(Radius * 0.25f, 0.0f, Height);
+
+	// Get list of unique polygons connected to the vertices that have moved
+	TArray<FPolygonID> ConnectedPolygons = MeshDescription.GetVertexConnectedPolygons(TopVertexIDs[0]);
+	for (const FPolygonID ConnectedPolygon : MeshDescription.GetVertexConnectedPolygons(BottomVertexIDs[0]))
+	{
+		ConnectedPolygons.AddUnique(ConnectedPolygon);
+	}
+
+	TestEqual("Number of polys to triangulate", ConnectedPolygons.Num(), 5);
+	TestTrue("PolygonsToTriangulate", ConnectedPolygons.Contains(TopPolygonID));
+	TestTrue("PolygonsToTriangulate", ConnectedPolygons.Contains(SidePolygonIDs[0]));
+	TestTrue("PolygonsToTriangulate", ConnectedPolygons.Contains(SidePolygonIDs[NumSides - 1]));
+	TestTrue("PolygonsToTriangulate", ConnectedPolygons.Contains(MeshDescription.GetTrianglePolygon(BottomTriangleIDs[0])));
+	TestTrue("PolygonsToTriangulate", ConnectedPolygons.Contains(MeshDescription.GetTrianglePolygon(BottomTriangleIDs[NumSides - 1])));
+
+	for (const FPolygonID ConnectedPolygon : ConnectedPolygons)
+	{
+		MeshDescription.ComputePolygonTriangulation(ConnectedPolygon);
+	}
+
+	TArray<FEdgeID, TInlineAllocator<NumSides - 3>> NewTopInternalEdgeIDs = MeshDescription.GetPolygonInternalEdges<TInlineAllocator<NumSides - 3>>(TopPolygonID);
+	TestEqual("Number of internal edges", NewTopInternalEdgeIDs.Num(), NumSides - 3);
+
+	// Move polygon to another group
+
+	MeshDescription.SetPolygonPolygonGroup(TopPolygonID, SidePolygonGroupID);
+	TestEqual("New polygons in side group", MeshDescription.GetNumPolygonGroupPolygons(SidePolygonGroupID), NumSides + 1);
+	TestEqual("New polygons in top group", MeshDescription.GetNumPolygonGroupPolygons(TopPolygonGroupID), 0);
+	TestEqual("Top polygon group", MeshDescription.GetPolygonPolygonGroup(TopPolygonID), SidePolygonGroupID);
+
+	// Delete a mesh element
+
+	TArray<FEdgeID> TopFaceEdges = MeshDescription.GetPolygonPerimeterEdges(TopPolygonID);
+	for (int32 Index = 0; Index < NumSides; ++Index)
+	{
+		const FEdgeID TopFaceEdge = TopFaceEdges[Index];
+		TestEqual("Check number of edge connected polygons", MeshDescription.GetNumEdgeConnectedPolygons(TopFaceEdge), 2);
+		TestTrue("Check edge connection", MeshDescription.GetEdgeConnectedPolygons(TopFaceEdge).Contains(TopPolygonID));
+		TestTrue("Check edge connection", MeshDescription.GetEdgeConnectedPolygons(TopFaceEdge).Contains(SidePolygonIDs[Index]));
+	}
+
+	TArray<FEdgeID> OrphanedEdges;
+	TArray<FVertexInstanceID> OrphanedVertexInstances;
+	TArray<FPolygonGroupID> OrphanedPolygonGroups;
+	MeshDescription.DeletePolygon(TopPolygonID, &OrphanedEdges, &OrphanedVertexInstances, &OrphanedPolygonGroups);
+
+	TestEqual("Orphaned edges", OrphanedEdges.Num(), 0);
+	TestEqual("Orphaned vertex instances", OrphanedVertexInstances.Num(), NumSides);
+	TestEqual("Orphaned polygon groups", OrphanedPolygonGroups.Num(), 0);
+	TestEqual("New edge count", MeshDescription.Edges().Num(), NumSides * 5);
+	TestEqual("New edge array size", MeshDescription.Edges().GetArraySize(), NumSides * 6 - 3);
+
+	for (int32 Index = 0; Index < NumSides; ++Index)
+	{
+		const FEdgeID TopFaceEdge = TopFaceEdges[Index];
+		TestEqual("Check number of edge connected polygons", MeshDescription.GetNumEdgeConnectedPolygons(TopFaceEdge), 1);
+		TestTrue("Check edge connection", MeshDescription.GetEdgeConnectedPolygons(TopFaceEdge).Contains(SidePolygonIDs[Index]));
+	}
+
+	return true;
+}
diff --git a/Engine/Source/Developer/MaterialBaking/Private/MaterialRenderItem.cpp b/Engine/Source/Developer/MaterialBaking/Private/MaterialRenderItem.cpp
--- a/Engine/Source/Developer/MaterialBaking/Private/MaterialRenderItem.cpp
+++ b/Engine/Source/Developer/MaterialBaking/Private/MaterialRenderItem.cpp
@@ -205,15 +205,15 @@ void FMeshMaterialRenderItem::PopulateWithMeshData()
 	for(const FPolygonID PolygonID : RawMesh->Polygons().GetElementIDs())
 	{
 		const FPolygonGroupID PolygonGroupID = RawMesh->GetPolygonPolygonGroup(PolygonID);
-		const TArray<FMeshTriangle>& Triangles = RawMesh->GetPolygonTriangles(PolygonID);
-		for (const FMeshTriangle& Triangle : Triangles)
+		const TArray<FTriangleID>& TriangleIDs = RawMesh->GetPolygonTriangleIDs(PolygonID);
+		for (const FTriangleID TriangleID : TriangleIDs)
 		{
 			if (MeshSettings->MaterialIndices.Contains(PolygonGroupID.GetValue()))
 			{
 				for (int32 Corner = 0; Corner < 3; Corner++)
 				{
 					const int32 SrcVertIndex = FaceIndex * 3 + Corner;
-					const FVertexInstanceID SrcVertexInstanceID = Triangle.GetVertexInstanceID(Corner);
+					const FVertexInstanceID SrcVertexInstanceID = RawMesh->GetTriangleVertexInstance(TriangleID, Corner);
 					const FVertexID SrcVertexID = RawMesh->GetVertexInstanceVertex(SrcVertexInstanceID);
 					// add vertex
 					FDynamicMeshVertex* Vert = new(Vertices)FDynamicMeshVertex();
diff --git a/Engine/Source/Developer/MaterialUtilities/Private/MeshRendering.cpp b/Engine/Source/Developer/MaterialUtilities/Private/MeshRendering.cpp
--- a/Engine/Source/Developer/MaterialUtilities/Private/MeshRendering.cpp
+++ b/Engine/Source/Developer/MaterialUtilities/Private/MeshRendering.cpp
@@ -158,7 +158,7 @@ public:
 			const FPolygonGroupID PolygonGroupID = RawMesh.GetPolygonPolygonGroup(PolygonID);
 			if (PolygonGroupID.GetValue() == Data.MaterialIndex)
 			{
-				NumTris += RawMesh.GetPolygonTriangles(PolygonID).Num();
+				NumTris += RawMesh.GetPolygonTriangleIDs(PolygonID).Num();
 			}
 		}
 		if (NumTris == 0)
@@ -205,15 +205,15 @@ public:
 		for (const FPolygonID PolygonID : RawMesh.Polygons().GetElementIDs())
 		{
 			const FPolygonGroupID PolygonGroupID = RawMesh.GetPolygonPolygonGroup(PolygonID);
-			const TArray<FMeshTriangle>& Triangles = RawMesh.GetPolygonTriangles(PolygonID);
-			for (const FMeshTriangle& Triangle : Triangles)
+			const TArray<FTriangleID>& TriangleIDs = RawMesh.GetPolygonTriangleIDs(PolygonID);
+			for (const FTriangleID TriangleID : TriangleIDs)
 			{
 				if (PolygonGroupID.GetValue() == Data.MaterialIndex)
 				{
 					for (int32 Corner = 0; Corner < 3; Corner++)
 					{
 						const int32 SrcVertIndex = FaceIndex * 3 + Corner;
-						const FVertexInstanceID SrcVertexInstanceID = Triangle.GetVertexInstanceID(Corner);
+						const FVertexInstanceID SrcVertexInstanceID = RawMesh.GetTriangleVertexInstance(TriangleID, Corner);
 						const FVertexID SrcVertexID = RawMesh.GetVertexInstanceVertex(SrcVertexInstanceID);
 
 						// add vertex
diff --git a/Engine/Source/Developer/MeshBuilder/Private/StaticMeshBuilder.cpp b/Engine/Source/Developer/MeshBuilder/Private/StaticMeshBuilder.cpp
--- a/Engine/Source/Developer/MeshBuilder/Private/StaticMeshBuilder.cpp
+++ b/Engine/Source/Developer/MeshBuilder/Private/StaticMeshBuilder.cpp
@@ -404,12 +404,7 @@ void BuildVertexBuffer(
 		}
 	}
 
-	int32 ReserveIndicesCount = 0;
-	for (const FPolygonID& PolygonID : MeshDescription.Polygons().GetElementIDs())
-	{
-		const TArray<FMeshTriangle>& PolygonTriangles = MeshDescription.GetPolygonTriangles(PolygonID);
-		ReserveIndicesCount += PolygonTriangles.Num() * 3;
-	}
+	int32 ReserveIndicesCount = MeshDescription.Triangles().Num() * 3;
 	IndexBuffer.Reset(ReserveIndicesCount);
 
 	//Fill the remap array
@@ -424,23 +419,23 @@ void BuildVertexBuffer(
 	OutWedgeMap.AddZeroed(ReserveIndicesCount);
 
 	int32 WedgeIndex = 0;
-	for (const FPolygonID& PolygonID : MeshDescription.Polygons().GetElementIDs())
+	for (const FPolygonID PolygonID : MeshDescription.Polygons().GetElementIDs())
 	{
 		const FPolygonGroupID PolygonGroupID = MeshDescription.GetPolygonPolygonGroup(PolygonID);
 		const int32 SectionIndex = PolygonGroupToSectionIndex[PolygonGroupID];
 		TArray<uint32>& SectionIndices = OutPerSectionIndices[SectionIndex];
 
-		const TArray<FMeshTriangle>& PolygonTriangles = MeshDescription.GetPolygonTriangles(PolygonID);
+		const TArray<FTriangleID>& TriangleIDs = MeshDescription.GetPolygonTriangleIDs(PolygonID);
 		uint32 MinIndex = TNumericLimits< uint32 >::Max();
 		uint32 MaxIndex = TNumericLimits< uint32 >::Min();
-		for (int32 TriangleIndex = 0; TriangleIndex < PolygonTriangles.Num(); ++TriangleIndex)
+		for (int32 TriangleIndex = 0; TriangleIndex < TriangleIDs.Num(); ++TriangleIndex)
 		{
-			const FMeshTriangle& Triangle = PolygonTriangles[TriangleIndex];
+			const FTriangleID TriangleID = TriangleIDs[TriangleIndex];
 
 			FVector CornerPositions[3];
 			for (int32 TriVert = 0; TriVert < 3; ++TriVert)
 			{
-				const FVertexInstanceID VertexInstanceID = Triangle.GetVertexInstanceID(TriVert);
+				const FVertexInstanceID VertexInstanceID = MeshDescription.GetTriangleVertexInstance(TriangleID, TriVert);
 				const FVertexID VertexID = MeshDescription.GetVertexInstanceVertex(VertexInstanceID);
 				CornerPositions[TriVert] = VertexPositions[VertexID];
 			}
@@ -457,7 +452,7 @@ void BuildVertexBuffer(
 
 			for (int32 TriVert = 0; TriVert < 3; ++TriVert, ++WedgeIndex)
 			{
-				const FVertexInstanceID VertexInstanceID = Triangle.GetVertexInstanceID(TriVert);
+				const FVertexInstanceID VertexInstanceID = MeshDescription.GetTriangleVertexInstance(TriangleID, TriVert);
 				const int32 VertexInstanceValue = VertexInstanceID.GetValue();
 				const FVector& VertexPosition = CornerPositions[TriVert];
 				const FVector& VertexInstanceNormal = VertexInstanceNormals[VertexInstanceID];
diff --git a/Engine/Source/Developer/MeshDescriptionOperations/Private/MeshDescriptionOperations.cpp b/Engine/Source/Developer/MeshDescriptionOperations/Private/MeshDescriptionOperations.cpp
--- a/Engine/Source/Developer/MeshDescriptionOperations/Private/MeshDescriptionOperations.cpp
+++ b/Engine/Source/Developer/MeshDescriptionOperations/Private/MeshDescriptionOperations.cpp
@@ -80,7 +80,7 @@ void FMeshDescriptionOperations::ConvertHardEdgesToSmoothGroup(const FMeshDescri
 	int32 TriangleCount = 0;
 	for (const FPolygonID PolygonID : SourceMeshDescription.Polygons().GetElementIDs())
 	{
-		TriangleCount += SourceMeshDescription.GetPolygonTriangles(PolygonID).Num();
+		TriangleCount += SourceMeshDescription.GetPolygonTriangleIDs(PolygonID).Num();
 		if (ConsumedPolygons[PolygonID.GetValue()])
 		{
 			continue;
@@ -104,7 +104,7 @@ void FMeshDescriptionOperations::ConvertHardEdgesToSmoothGroup(const FMeshDescri
 			uint32 NeighborSmoothGroup = 0;
 			const uint32 LastSmoothGroupValue = (LastPolygonID == FPolygonID::Invalid) ? 0 : PolygonSmoothGroup[LastPolygonID];
 			TArray<FEdgeID> PolygonEdges;
-			SourceMeshDescription.GetPolygonEdges(CurrentPolygonID, PolygonEdges);
+			SourceMeshDescription.GetPolygonPerimeterEdges(CurrentPolygonID, PolygonEdges);
 			for (const FEdgeID& EdgeID : PolygonEdges)
 			{
 				bool bIsHardEdge = EdgeHardnesses[EdgeID];
@@ -184,8 +184,7 @@ void FMeshDescriptionOperations::ConvertHardEdgesToSmoothGroup(const FMeshDescri
 	for (const FPolygonID PolygonID : SourceMeshDescription.Polygons().GetElementIDs())
 	{
 		uint32 PolygonSmoothValue = PolygonSmoothGroup[PolygonID];
-		const TArray<FMeshTriangle>& Triangles = SourceMeshDescription.GetPolygonTriangles(PolygonID);
-		for (const FMeshTriangle& MeshTriangle : Triangles)
+		for (const FTriangleID TriangleID : SourceMeshDescription.GetPolygonTriangleIDs(PolygonID))
 		{
 			FaceSmoothingMasks[TriangleIndex++] = PolygonSmoothValue;
 		}
@@ -213,7 +212,7 @@ void FMeshDescriptionOperations::ConvertSmoothGroupToHardEdges(const TArray<uint
 			check(FaceSmoothingMasks.IsValidIndex(CurrentPolygonIDValue));
 			const uint32 ReferenceSmoothGroup = FaceSmoothingMasks[CurrentPolygonIDValue];
 			TArray<FEdgeID> PolygonEdges;
-			DestinationMeshDescription.GetPolygonEdges(CurrentPolygonID, PolygonEdges);
+			DestinationMeshDescription.GetPolygonPerimeterEdges(CurrentPolygonID, PolygonEdges);
 			for (const FEdgeID& EdgeID : PolygonEdges)
 			{
 				const bool bIsHardEdge = EdgeHardnesses[EdgeID];
@@ -273,11 +272,7 @@ void FMeshDescriptionOperations::ConvertToRawMesh(const FMeshDescription& Source
 		++VertexIndex;
 	}
 
-	int32 TriangleNumber = 0;
-	for (const FPolygonID& PolygonID : SourceMeshDescription.Polygons().GetElementIDs())
-	{
-		TriangleNumber += SourceMeshDescription.GetPolygonTriangles(PolygonID).Num();
-	}
+	int32 TriangleNumber = SourceMeshDescription.Triangles().Num();
 	DestinationRawMesh.FaceMaterialIndices.AddZeroed(TriangleNumber);
 	DestinationRawMesh.FaceSmoothingMasks.AddZeroed(TriangleNumber);
 
@@ -304,8 +299,8 @@ void FMeshDescriptionOperations::ConvertToRawMesh(const FMeshDescription& Source
 	{
 		const FPolygonGroupID& PolygonGroupID = SourceMeshDescription.GetPolygonPolygonGroup(PolygonID);
 		int32 PolygonIDValue = PolygonID.GetValue();
-		const TArray<FMeshTriangle>& Triangles = SourceMeshDescription.GetPolygonTriangles(PolygonID);
-		for (const FMeshTriangle& MeshTriangle : Triangles)
+		const TArray<FTriangleID>& TriangleIDs = SourceMeshDescription.GetPolygonTriangleIDs(PolygonID);
+		for (const FTriangleID TriangleID : TriangleIDs)
 		{
 			if (MaterialMap.Num() > 0 && MaterialMap.Contains(PolygonGroupMaterialSlotName[PolygonGroupID]))
 			{
@@ -318,7 +313,7 @@ void FMeshDescriptionOperations::ConvertToRawMesh(const FMeshDescription& Source
 			DestinationRawMesh.FaceSmoothingMasks[TriangleIndex] = 0; //Conversion of soft/hard to smooth mask is done after the geometry is converted
 			for (int32 Corner = 0; Corner < 3; ++Corner)
 			{
-				const FVertexInstanceID VertexInstanceID = MeshTriangle.GetVertexInstanceID(Corner);
+				const FVertexInstanceID VertexInstanceID = SourceMeshDescription.GetTriangleVertexInstance(TriangleID, Corner);
 
 				if (bHasVertexColor)
 				{
@@ -555,14 +550,7 @@ void FMeshDescriptionOperations::ConvertFromRawMesh(const FRawMesh& SourceRawMes
 			}
 		}
 
-		const FPolygonID NewPolygonID = DestinationMeshDescription.CreatePolygon(PolygonGroupID, TriangleVertexInstanceIDs);
-		int32 NewTriangleIndex = DestinationMeshDescription.GetPolygonTriangles(NewPolygonID).AddDefaulted();
-		FMeshTriangle& NewTriangle = DestinationMeshDescription.GetPolygonTriangles(NewPolygonID)[NewTriangleIndex];
-		for (int32 Corner = 0; Corner < 3; ++Corner)
-		{
-			FVertexInstanceID VertexInstanceID = TriangleVertexInstanceIDs[Corner];
-			NewTriangle.SetVertexInstanceID(Corner, VertexInstanceID);
-		}
+		DestinationMeshDescription.CreatePolygon(PolygonGroupID, TriangleVertexInstanceIDs);
 	}
 
 	ConvertSmoothGroupToHardEdges(SourceRawMesh.FaceSmoothingMasks, DestinationMeshDescription);
@@ -682,10 +670,11 @@ void FMeshDescriptionOperations::AppendMeshDescription(const FMeshDescription& S
 	//Edges
 	TMap<FEdgeID, FEdgeID> SourceEdgeIDRemap;
 	SourceEdgeIDRemap.Reserve(SourceMesh.Edges().Num());
-	for (const FEdgeID& SourceEdgeID : SourceMesh.Edges().GetElementIDs())
+	for (const FEdgeID SourceEdgeID : SourceMesh.Edges().GetElementIDs())
 	{
-		const FMeshEdge& SourceEdge = SourceMesh.GetEdge(SourceEdgeID);
-		FEdgeID TargetEdgeID = TargetMesh.CreateEdge(SourceVertexIDRemap[SourceEdge.VertexIDs[0]], SourceVertexIDRemap[SourceEdge.VertexIDs[1]]);
+		const FVertexID EdgeVertex0 = SourceMesh.GetEdgeVertex(SourceEdgeID, 0);
+		const FVertexID EdgeVertex1 = SourceMesh.GetEdgeVertex(SourceEdgeID, 1);
+		FEdgeID TargetEdgeID = TargetMesh.CreateEdge(SourceVertexIDRemap[EdgeVertex0], SourceVertexIDRemap[EdgeVertex1]);
 		TargetEdgeHardnesses[TargetEdgeID] = SourceEdgeHardnesses[SourceEdgeID];
 		TargetEdgeCreaseSharpnesses[TargetEdgeID] = SourceEdgeCreaseSharpnesses[SourceEdgeID];
 		SourceEdgeIDRemap.Add(SourceEdgeID, TargetEdgeID);
@@ -735,24 +724,22 @@ void FMeshDescriptionOperations::AppendMeshDescription(const FMeshDescription& S
 	}
 
 	//Polygons
-	for (const FPolygonID& SourcePolygonID : SourceMesh.Polygons().GetElementIDs())
+	for (const FPolygonID SourcePolygonID : SourceMesh.Polygons().GetElementIDs())
 	{
-		const FMeshPolygon& SourcePolygon = SourceMesh.GetPolygon(SourcePolygonID);
+		const TArray<FVertexInstanceID>& PerimeterVertexInstanceIDs = SourceMesh.GetPolygonVertexInstances(SourcePolygonID);
+		const FPolygonGroupID PolygonGroupID = SourceMesh.GetPolygonPolygonGroup(SourcePolygonID);
 		//Find the polygonGroupID
-		FPolygonGroupID TargetPolygonGroupID = RemapPolygonGroup[SourcePolygon.PolygonGroupID];
+		FPolygonGroupID TargetPolygonGroupID = RemapPolygonGroup[PolygonGroupID];
 
-		int32 PolygonVertexCount = SourcePolygon.PerimeterContour.VertexInstanceIDs.Num();
+		int32 PolygonVertexCount = PerimeterVertexInstanceIDs.Num();
 		TArray<FVertexInstanceID> VertexInstanceIDs;
 		VertexInstanceIDs.Reserve(PolygonVertexCount);
-		for (const FVertexInstanceID VertexInstanceID : SourcePolygon.PerimeterContour.VertexInstanceIDs)
+		for (const FVertexInstanceID VertexInstanceID : PerimeterVertexInstanceIDs)
 		{
 			VertexInstanceIDs.Add(SourceVertexInstanceIDRemap[VertexInstanceID]);
 		}
 		// Insert a polygon into the mesh
 		const FPolygonID TargetPolygonID = TargetMesh.CreatePolygon(TargetPolygonGroupID, VertexInstanceIDs);
-		//Triangulate the polygon
-		FMeshPolygon& Polygon = TargetMesh.GetPolygon(TargetPolygonID);
-		TargetMesh.ComputePolygonTriangulation(TargetPolygonID, Polygon.Triangles);
 	}
 }
 
@@ -829,11 +816,11 @@ void FMeshDescriptionOperations::CreatePolygonNTB(FMeshDescription& MeshDescript
 			//By pass normal calculation if its already done
 			continue;
 		}
-		const TArray<FMeshTriangle>& MeshTriangles = MeshDescription.GetPolygonTriangles(PolygonID);
+		const TArray<FTriangleID>& TriangleIDs = MeshDescription.GetPolygonTriangleIDs(PolygonID);
 		FVector TangentX(0.0f);
 		FVector TangentY(0.0f);
 		FVector TangentZ(0.0f);
-		for (const FMeshTriangle& MeshTriangle : MeshTriangles)
+		for (const FTriangleID TriangleID : TriangleIDs)
 		{
 			int32 UVIndex = 0;
 
@@ -842,7 +829,7 @@ void FMeshDescriptionOperations::CreatePolygonNTB(FMeshDescription& MeshDescript
 
 			for (int32 i = 0; i < 3; ++i)
 			{
-				const FVertexInstanceID VertexInstanceID = MeshTriangle.GetVertexInstanceID(i);
+				const FVertexInstanceID VertexInstanceID = MeshDescription.GetTriangleVertexInstance(TriangleID, i);
 				UVs[i] = VertexUVs.Get(VertexInstanceID, 0); // UV0
 				P[i] = VertexPositions[MeshDescription.GetVertexInstanceVertex(VertexInstanceID)];
 			}
@@ -898,7 +885,7 @@ void FMeshDescriptionOperations::CreatePolygonNTB(FMeshDescription& MeshDescript
 
 		// Calculate polygon center: just an average of all vertex positions.
 		FVector Center = FVector::ZeroVector;
-		const TArray<FVertexInstanceID>& VertexInstanceIDs = MeshDescription.GetPolygonPerimeterVertexInstances(PolygonID);
+		const TArray<FVertexInstanceID>& VertexInstanceIDs = MeshDescription.GetPolygonVertexInstances(PolygonID);
 		for (const FVertexInstanceID VertexInstanceID : VertexInstanceIDs)
 		{
 			Center += VertexPositions[MeshDescription.GetVertexInstanceVertex(VertexInstanceID)];
@@ -970,7 +957,7 @@ void FMeshDescriptionOperations::CreateNormals(FMeshDescription& MeshDescription
 				if (VertexInfo.PolygonID == FPolygonID::Invalid)
 				{
 					VertexInfo.PolygonID = PolygonID;
-					for (const FVertexInstanceID VertexInstanceID : MeshDescription.GetPolygonPerimeterVertexInstances(PolygonID))
+					for (const FVertexInstanceID VertexInstanceID : MeshDescription.GetPolygonVertexInstances(PolygonID))
 					{
 						if (MeshDescription.GetVertexInstanceVertex(VertexInstanceID) == VertexID)
 						{
@@ -1209,8 +1196,7 @@ namespace MeshDescriptionMikktSpaceInterface
 		FMeshDescription *MeshDescription = (FMeshDescription*)(Context->m_pUserData);
 		if (MeshDescription->IsPolygonValid(FPolygonID(FaceIdx)))
 		{
-			const FMeshPolygon& Polygon = MeshDescription->GetPolygon(FPolygonID(FaceIdx));
-			return Polygon.PerimeterContour.VertexInstanceIDs.Num();
+			return MeshDescription->GetPolygonVertexInstances(FPolygonID(FaceIdx)).Num();
 		}
 
 		return 0;
@@ -1219,8 +1205,8 @@ namespace MeshDescriptionMikktSpaceInterface
 	void MikkGetPosition(const SMikkTSpaceContext* Context, float Position[3], const int FaceIdx, const int VertIdx)
 	{
 		FMeshDescription* MeshDescription = (FMeshDescription*)(Context->m_pUserData);
-		const FMeshPolygon& Polygon = MeshDescription->GetPolygon(FPolygonID(FaceIdx));
-		const FVertexInstanceID VertexInstanceID = Polygon.PerimeterContour.VertexInstanceIDs[VertIdx];
+		const TArray<FVertexInstanceID>& VertexInstanceIDs = MeshDescription->GetPolygonVertexInstances(FPolygonID(FaceIdx));
+		const FVertexInstanceID VertexInstanceID = VertexInstanceIDs[VertIdx];
 		const FVertexID VertexID = MeshDescription->GetVertexInstanceVertex(VertexInstanceID);
 		const FVector& VertexPosition = MeshDescription->VertexAttributes().GetAttribute<FVector>(VertexID, MeshAttribute::Vertex::Position);
 		Position[0] = VertexPosition.X;
@@ -1231,8 +1217,8 @@ namespace MeshDescriptionMikktSpaceInterface
 	void MikkGetNormal(const SMikkTSpaceContext* Context, float Normal[3], const int FaceIdx, const int VertIdx)
 	{
 		FMeshDescription* MeshDescription = (FMeshDescription*)(Context->m_pUserData);
-		const FMeshPolygon& Polygon = MeshDescription->GetPolygon(FPolygonID(FaceIdx));
-		const FVertexInstanceID VertexInstanceID = Polygon.PerimeterContour.VertexInstanceIDs[VertIdx];
+		const TArray<FVertexInstanceID>& VertexInstanceIDs = MeshDescription->GetPolygonVertexInstances(FPolygonID(FaceIdx));
+		const FVertexInstanceID VertexInstanceID = VertexInstanceIDs[VertIdx];
 		const FVector& VertexNormal = MeshDescription->VertexInstanceAttributes().GetAttribute<FVector>(VertexInstanceID, MeshAttribute::VertexInstance::Normal);
 		Normal[0] = VertexNormal.X;
 		Normal[1] = VertexNormal.Y;
@@ -1242,8 +1228,8 @@ namespace MeshDescriptionMikktSpaceInterface
 	void MikkSetTSpaceBasic(const SMikkTSpaceContext* Context, const float Tangent[3], const float BitangentSign, const int FaceIdx, const int VertIdx)
 	{
 		FMeshDescription* MeshDescription = (FMeshDescription*)(Context->m_pUserData);
-		const FMeshPolygon& Polygon = MeshDescription->GetPolygon(FPolygonID(FaceIdx));
-		const FVertexInstanceID VertexInstanceID = Polygon.PerimeterContour.VertexInstanceIDs[VertIdx];
+		const TArray<FVertexInstanceID>& VertexInstanceIDs = MeshDescription->GetPolygonVertexInstances(FPolygonID(FaceIdx));
+		const FVertexInstanceID VertexInstanceID = VertexInstanceIDs[VertIdx];
 		const FVector VertexTangent(Tangent[0], Tangent[1], Tangent[2]);
 		MeshDescription->VertexInstanceAttributes().SetAttribute<FVector>(VertexInstanceID, MeshAttribute::VertexInstance::Tangent, 0, VertexTangent);
 		MeshDescription->VertexInstanceAttributes().SetAttribute<float>(VertexInstanceID, MeshAttribute::VertexInstance::BinormalSign, 0, -BitangentSign);
@@ -1252,8 +1238,8 @@ namespace MeshDescriptionMikktSpaceInterface
 	void MikkGetTexCoord(const SMikkTSpaceContext* Context, float UV[2], const int FaceIdx, const int VertIdx)
 	{
 		FMeshDescription* MeshDescription = (FMeshDescription*)(Context->m_pUserData);
-		const FMeshPolygon& Polygon = MeshDescription->GetPolygon(FPolygonID(FaceIdx));
-		const FVertexInstanceID VertexInstanceID = Polygon.PerimeterContour.VertexInstanceIDs[VertIdx];
+		const TArray<FVertexInstanceID>& VertexInstanceIDs = MeshDescription->GetPolygonVertexInstances(FPolygonID(FaceIdx));
+		const FVertexInstanceID VertexInstanceID = VertexInstanceIDs[VertIdx];
 		const FVector2D& TexCoord = MeshDescription->VertexInstanceAttributes().GetAttribute<FVector2D>(VertexInstanceID, MeshAttribute::VertexInstance::TextureCoordinate, 0);
 		UV[0] = TexCoord.X;
 		UV[1] = TexCoord.Y;
@@ -1267,11 +1253,7 @@ void FMeshDescriptionOperations::FindOverlappingCorners(FOverlappingCorners& Out
 	const FVertexInstanceArray& VertexInstanceArray = MeshDescription.VertexInstances();
 	const FVertexArray& VertexArray = MeshDescription.Vertices();
 
-	int32 NumWedges = 0;
-	for (const FPolygonID PolygonID : MeshDescription.Polygons().GetElementIDs())
-	{
-		NumWedges += 3 * MeshDescription.GetPolygonTriangles(PolygonID).Num();
-	}
+	int32 NumWedges = 3 * MeshDescription.Triangles().Num();
 
 	// Empty the old data and reserve space for new
 	OutOverlappingCorners.Init(NumWedges);
@@ -1285,12 +1267,12 @@ void FMeshDescriptionOperations::FindOverlappingCorners(FOverlappingCorners& Out
 	int32 WedgeIndex = 0;
 	for (const FPolygonID PolygonID : MeshDescription.Polygons().GetElementIDs())
 	{
-		const TArray<FMeshTriangle>& Triangles = MeshDescription.GetPolygonTriangles(PolygonID);
-		for (const FMeshTriangle& MeshTriangle : Triangles)
+		const TArray<FTriangleID>& TriangleIDs = MeshDescription.GetPolygonTriangleIDs(PolygonID);
+		for (const FTriangleID TriangleID : TriangleIDs)
 		{
 			for (int32 Corner = 0; Corner < 3; ++Corner)
 			{
-				const FVertexInstanceID VertexInstanceID = MeshTriangle.GetVertexInstanceID(Corner);
+				const FVertexInstanceID VertexInstanceID = MeshDescription.GetTriangleVertexInstance(TriangleID, Corner);
 				new(VertIndexAndZ)MeshDescriptionOperationNamespace::FIndexAndZ(WedgeIndex, VertexPositions[MeshDescription.GetVertexInstanceVertex(VertexInstanceID)]);
 				++WedgeIndex;
 			}
@@ -1345,11 +1327,7 @@ struct FLayoutUVMeshDescriptionView final : FLayoutUV::IMeshView
 		, SrcChannel(InSrcChannel)
 		, DstChannel(InDstChannel)
 	{
-		uint32 NumTris = 0;
-		for (const FPolygonID PolygonID : MeshDescription.Polygons().GetElementIDs())
-		{
-			NumTris += MeshDescription.GetPolygonTriangles(PolygonID).Num();
-		}
+		uint32 NumTris = MeshDescription.Triangles().Num();
 
 		NumIndices = NumTris * 3;
 
@@ -1360,12 +1338,12 @@ struct FLayoutUVMeshDescriptionView final : FLayoutUV::IMeshView
 
 		for (const FPolygonID PolygonID : MeshDescription.Polygons().GetElementIDs())
 		{
-			const TArray<FMeshTriangle>& Triangles = MeshDescription.GetPolygonTriangles(PolygonID);
-			for (const FMeshTriangle MeshTriangle : Triangles)
+			const TArray<FTriangleID>& TriangleIDs = MeshDescription.GetPolygonTriangleIDs(PolygonID);
+			for (const FTriangleID TriangleID : TriangleIDs)
 			{
 				for (int32 Corner = 0; Corner < 3; ++Corner)
 				{
-					const FVertexInstanceID VertexInstanceID = MeshTriangle.GetVertexInstanceID(Corner);
+					const FVertexInstanceID VertexInstanceID = MeshDescription.GetTriangleVertexInstance(TriangleID, Corner);
 
 					FlattenedTexCoords[WedgeIndex] = TexCoords.Get(VertexInstanceID, SrcChannel);
 					RemapVerts[WedgeIndex] = VertexInstanceID.GetValue();
@@ -1468,7 +1446,7 @@ bool FMeshDescriptionOperations::GenerateUniqueUVsForStaticMesh(const FMeshDescr
 		for (FPolygonID RefPolygonID : DuplicateMeshDescription.Polygons().GetElementIDs())
 		{
 			FPolygonGroupID RefPolygonGroupID = DuplicateMeshDescription.GetPolygonPolygonGroup(RefPolygonID);
-			const TArray<FVertexInstanceID>& RefVertexInstances = DuplicateMeshDescription.GetPolygonPerimeterVertexInstances(RefPolygonID);
+			const TArray<FVertexInstanceID>& RefVertexInstances = DuplicateMeshDescription.GetPolygonVertexInstances(RefPolygonID);
 			TArray<FVector2D> RefUVs;
 			for (FVertexInstanceID RefVertexInstanceID : RefVertexInstances)
 			{
@@ -1482,7 +1460,7 @@ bool FMeshDescriptionOperations::GenerateUniqueUVsForStaticMesh(const FMeshDescr
 				{
 					continue;
 				}
-				const TArray<FVertexInstanceID>& TestVertexInstances = DuplicateMeshDescription.GetPolygonPerimeterVertexInstances(TestPolygonID);
+				const TArray<FVertexInstanceID>& TestVertexInstances = DuplicateMeshDescription.GetPolygonVertexInstances(TestPolygonID);
 				if (TestVertexInstances.Num() != RefVertexInstances.Num())
 				{
 					continue;
@@ -1515,7 +1493,7 @@ bool FMeshDescriptionOperations::GenerateUniqueUVsForStaticMesh(const FMeshDescr
 			}
 			else
 			{
-				const TArray<FVertexInstanceID>& TestVertexInstances = DuplicateMeshDescription.GetPolygonPerimeterVertexInstances(MatchPolygonID);
+				const TArray<FVertexInstanceID>& TestVertexInstances = DuplicateMeshDescription.GetPolygonVertexInstances(MatchPolygonID);
 				int32 VertexInstanceIndex = 0;
 				for (FVertexInstanceID RefVertexInstanceID : RefVertexInstances)
 				{
@@ -1709,7 +1687,7 @@ void FMeshDescriptionOperations::GenerateCylindricalUV(FMeshDescription& MeshDes
 	// Fix the UV coordinates for triangles at the seam where the angle wraps around
 	for (const FPolygonID& PolygonID : MeshDescription.Polygons().GetElementIDs())
 	{
-		const TArray<FVertexInstanceID>& VertexInstances = MeshDescription.GetPolygonPerimeterVertexInstances(PolygonID);
+		const TArray<FVertexInstanceID>& VertexInstances = MeshDescription.GetPolygonVertexInstances(PolygonID);
 		int32 NumInstances = VertexInstances.Num();
 		if (NumInstances >= 2)
 		{
@@ -1774,7 +1752,7 @@ void FMeshDescriptionOperations::GenerateBoxUV(const FMeshDescription& MeshDescr
 	// For each polygon, find the box plane that best matches the polygon normal
 	for (const FPolygonID& PolygonID : MeshDescription.Polygons().GetElementIDs())
 	{
-		const TArray<FVertexInstanceID>& VertexInstances = MeshDescription.GetPolygonPerimeterVertexInstances(PolygonID);
+		const TArray<FVertexInstanceID>& VertexInstances = MeshDescription.GetPolygonVertexInstances(PolygonID);
 		check(VertexInstances.Num() == 3);
 
 		FVector Vertex0 = VertexPositions[MeshDescription.GetVertexInstanceVertex(VertexInstances[0])];
@@ -1821,45 +1799,7 @@ void FMeshDescriptionOperations::GenerateBoxUV(const FMeshDescription& MeshDescr
 
 void FMeshDescriptionOperations::RemapPolygonGroups(FMeshDescription& MeshDescription, TMap<FPolygonGroupID, FPolygonGroupID>& Remap)
 {
-	TPolygonGroupAttributesRef<FName> PolygonGroupNames = MeshDescription.PolygonGroupAttributes().GetAttributesRef<FName>(MeshAttribute::PolygonGroup::ImportedMaterialSlotName);
-
-	struct FOldPolygonGroupData
-	{
-		FName Name;
-		TArray<FPolygonID> Polygons;
-	};
-
-	TMap<FPolygonGroupID, FOldPolygonGroupData> OldData;
-	for (const FPolygonGroupID& PolygonGroupID : MeshDescription.PolygonGroups().GetElementIDs())
-	{
-		if (!Remap.Contains(PolygonGroupID) || PolygonGroupID == Remap[PolygonGroupID])
-		{
-			//No need to change this one
-			continue;
-		}
-		FOldPolygonGroupData& PolygonGroupData = OldData.FindOrAdd(PolygonGroupID);
-		PolygonGroupData.Name = PolygonGroupNames[PolygonGroupID];
-		FMeshPolygonGroup& PolygonGroup = MeshDescription.GetPolygonGroup(PolygonGroupID);
-		PolygonGroupData.Polygons = PolygonGroup.Polygons;
-		PolygonGroup.Polygons.Empty();
-		MeshDescription.DeletePolygonGroup(PolygonGroupID);
-	}
-	for (auto Kvp : OldData)
-	{
-		FPolygonGroupID GroupID = Kvp.Key;
-		FPolygonGroupID ToGroupID = Remap[GroupID];
-		if (!MeshDescription.PolygonGroups().IsValid(ToGroupID))
-		{
-			MeshDescription.CreatePolygonGroupWithID(ToGroupID);
-		}
-		TArray<FPolygonID>& Polygons = MeshDescription.GetPolygonGroup(ToGroupID).Polygons;
-		Polygons.Append(Kvp.Value.Polygons);
-		PolygonGroupNames[ToGroupID] = Kvp.Value.Name;
-		for (FPolygonID PolygonID : Polygons)
-		{
-			MeshDescription.GetPolygon(PolygonID).PolygonGroupID = ToGroupID;
-		}
-	}
+	MeshDescription.RemapPolygonGroups(Remap);
 }
 
 void FMeshDescriptionOperations::SwapPolygonPolygonGroup(FMeshDescription& MeshDescription, int32 SectionIndex, int32 TriangleIndexStart, int32 TriangleIndexEnd, bool bRemoveEmptyPolygonGroup)
@@ -1889,12 +1829,11 @@ void FMeshDescriptionOperations::SwapPolygonPolygonGroup(FMeshDescription& MeshD
 
 	for (const FPolygonID PolygonID : MeshDescription.Polygons().GetElementIDs())
 	{
-		int32 TriangleCount = MeshDescription.GetPolygonTriangles(PolygonID).Num();
+		int32 TriangleCount = MeshDescription.GetPolygonTriangleIDs(PolygonID).Num();
 		if (TriangleIndex >= TriangleIndexStart && TriangleIndex < TriangleIndexEnd)
 		{
 			check(TriangleIndex + (TriangleCount-1) < TriangleIndexEnd);
-			FMeshPolygon& Polygon = MeshDescription.GetPolygon(PolygonID);
-			FPolygonGroupID OldpolygonGroupID = Polygon.PolygonGroupID;
+			FPolygonGroupID OldpolygonGroupID = MeshDescription.GetPolygonPolygonGroup(PolygonID);
 			if (OldpolygonGroupID != TargetPolygonGroupID)
 			{
 				MeshDescription.SetPolygonPolygonGroup(PolygonID, TargetPolygonGroupID);
diff --git a/Engine/Source/Developer/MeshMergeUtilities/Private/MeshMergeHelpers.cpp b/Engine/Source/Developer/MeshMergeUtilities/Private/MeshMergeHelpers.cpp
--- a/Engine/Source/Developer/MeshMergeUtilities/Private/MeshMergeHelpers.cpp
+++ b/Engine/Source/Developer/MeshMergeUtilities/Private/MeshMergeHelpers.cpp
@@ -357,9 +357,6 @@ void FMeshMergeHelpers::RetrieveMesh(USkeletalMeshComponent* SkeletalMeshCompone
 				}
 				//Create a polygon from this triangle
 				const FPolygonID NewPolygonID = RawMesh.CreatePolygon(SectionPolygonGroupID, VertexInstanceIDs);
-				//Triangulate the polygon
-				FMeshPolygon& Polygon = RawMesh.GetPolygon(NewPolygonID);
-				RawMesh.ComputePolygonTriangulation(NewPolygonID, Polygon.Triangles);
 			}
 		}
 	}
@@ -516,9 +513,6 @@ void FMeshMergeHelpers::ExportStaticMeshLOD(const FStaticMeshLODResources& Stati
 		}
 		//Create a polygon from this triangle
 		const FPolygonID NewPolygonID = OutRawMesh.CreatePolygon(CurrentPolygonGroupID, VertexInstanceIDs);
-		//Triangulate the polygon
-		FMeshPolygon& Polygon = OutRawMesh.GetPolygon(NewPolygonID);
-		OutRawMesh.ComputePolygonTriangulation(NewPolygonID, Polygon.Triangles);
 	}
 }
 
@@ -659,14 +653,14 @@ void FMeshMergeHelpers::CullTrianglesFromVolumesAndUnderLandscapes(const UWorld*
 
 	// We now know which vertices are below the landscape
 	TArray<FPolygonID> PolygonToRemove;
-	for(const FPolygonID& PolygonID : InOutRawMesh.Polygons().GetElementIDs())
+	for(const FPolygonID PolygonID : InOutRawMesh.Polygons().GetElementIDs())
 	{
 		bool AboveLandscape = false;
-		for (FMeshTriangle& Triangle : InOutRawMesh.GetPolygonTriangles(PolygonID))
+		for (const FTriangleID TriangleID : InOutRawMesh.GetPolygonTriangleIDs(PolygonID))
 		{
 			for (int32 Corner = 0; Corner < 3; ++Corner)
 			{
-				AboveLandscape |= VertexVisible[InOutRawMesh.GetVertexInstanceVertex(Triangle.GetVertexInstanceID(Corner))];
+				AboveLandscape |= VertexVisible[InOutRawMesh.GetVertexInstanceVertex(InOutRawMesh.GetTriangleVertexInstance(TriangleID, Corner))];
 			}
 		}
 		if (!AboveLandscape)
@@ -718,13 +712,13 @@ void FMeshMergeHelpers::PropagateSplineDeformationToRawMesh(const USplineMeshCom
 
 	// Apply spline deformation for each vertex's tangents
 	int32 WedgeIndex = 0;
-	for (const FPolygonID& PolygonID : OutRawMesh.Polygons().GetElementIDs())
+	for (const FPolygonID PolygonID : OutRawMesh.Polygons().GetElementIDs())
 	{
-		for (FMeshTriangle& Triangle : OutRawMesh.GetPolygonTriangles(PolygonID))
+		for (const FTriangleID TriangleID : OutRawMesh.GetPolygonTriangleIDs(PolygonID))
 		{
 			for (int32 Corner = 0; Corner < 3; ++Corner, ++WedgeIndex)
 			{
-				const FVertexInstanceID VertexInstanceID = Triangle.GetVertexInstanceID(Corner);
+				const FVertexInstanceID VertexInstanceID = OutRawMesh.GetTriangleVertexInstance(TriangleID, Corner);
 				const FVertexID VertexID = OutRawMesh.GetVertexInstanceVertex(VertexInstanceID);
 				const float& AxisValue = USplineMeshComponent::GetAxisValue(VertexPositions[VertexID], InSplineMeshComponent->ForwardAxis);
 				FTransform SliceTransform = InSplineMeshComponent->CalcSliceTransform(AxisValue);
@@ -1031,7 +1025,7 @@ void FMeshMergeHelpers::CalculateTextureCoordinateBoundsForRawMesh(const FMeshDe
 		if (OutBounds.Num() <= MaterialIndex)
 			OutBounds.SetNumZeroed(MaterialIndex + 1);
 		{
-			TArray<FVertexInstanceID> PolygonVertexInstances = InRawMesh.GetPolygonPerimeterVertexInstances(PolygonID);
+			TArray<FVertexInstanceID> PolygonVertexInstances = InRawMesh.GetPolygonVertexInstances(PolygonID);
 			for (const FVertexInstanceID& VertexInstanceID : PolygonVertexInstances)
 			{
 				for (int32 UVIndex = 0; UVIndex < VertexInstanceUVs.GetNumIndices(); ++UVIndex)
@@ -1068,14 +1062,14 @@ bool FMeshMergeHelpers::PropagatePaintedColorsToRawMesh(const UStaticMeshCompone
 				TMap<int32, FVertexInstanceID> IndexToVertexInstanceID;
 				IndexToVertexInstanceID.Reserve(NumWedges);
 				int32 CurrentWedgeIndex = 0;
-				for (const FPolygonID& PolygonID : RawMesh.Polygons().GetElementIDs())
+				for (const FPolygonID PolygonID : RawMesh.Polygons().GetElementIDs())
 				{
-					const TArray<FMeshTriangle>& Triangles = RawMesh.GetPolygonTriangles(PolygonID);
-					for (const FMeshTriangle& Triangle : Triangles)
+					const TArray<FTriangleID>& TriangleIDs = RawMesh.GetPolygonTriangleIDs(PolygonID);
+					for (const FTriangleID TriangleID : TriangleIDs)
 					{
 						for (int32 Corner = 0; Corner < 3; ++Corner, ++CurrentWedgeIndex)
 						{
-							IndexToVertexInstanceID.Add(CurrentWedgeIndex, Triangle.GetVertexInstanceID(Corner));
+							IndexToVertexInstanceID.Add(CurrentWedgeIndex, RawMesh.GetTriangleVertexInstance(TriangleID, Corner));
 						}
 					}
 				}
@@ -1254,8 +1248,9 @@ void FMeshMergeHelpers::AppendRawMesh(FMeshDescription& InTarget, const FMeshDes
 	SourceToTargetEdgeID.Reserve(InSource.Edges().Num());
 	for (const FEdgeID& SourceEdgeID : InSource.Edges().GetElementIDs())
 	{
-		const FMeshEdge& SourceEdge = InSource.GetEdge(SourceEdgeID);
-		const FEdgeID TargetEdgeID = InTarget.CreateEdge(SourceToTargetVertexID[SourceEdge.VertexIDs[0]], SourceToTargetVertexID[SourceEdge.VertexIDs[1]]);
+		const FVertexID EdgeVertex0 = InSource.GetEdgeVertex(SourceEdgeID, 0);
+		const FVertexID EdgeVertex1 = InSource.GetEdgeVertex(SourceEdgeID, 1);
+		const FEdgeID TargetEdgeID = InTarget.CreateEdge(SourceToTargetVertexID[EdgeVertex0], SourceToTargetVertexID[EdgeVertex1]);
 		TargetEdgeHardnesses[TargetEdgeID] = SourceEdgeHardnesses[SourceEdgeID];
 		TargetEdgeCreaseSharpnesses[TargetEdgeID] = SourceEdgeCreaseSharpnesses[SourceEdgeID];
 		SourceToTargetEdgeID.Add(SourceEdgeID, TargetEdgeID);
@@ -1272,20 +1267,16 @@ void FMeshMergeHelpers::AppendRawMesh(FMeshDescription& InTarget, const FMeshDes
 	};
 
 	//Append Polygons
-	for (const FPolygonID& SourcePolygonID : InSource.Polygons().GetElementIDs())
+	for (const FPolygonID SourcePolygonID : InSource.Polygons().GetElementIDs())
 	{
-		const FMeshPolygon& SourcePolygon = InSource.GetPolygon(SourcePolygonID);
-		const TArray<FVertexInstanceID>& SourceVertexInstanceIDs = InSource.GetPolygonPerimeterVertexInstances(SourcePolygonID);
-
+		const TArray<FVertexInstanceID>& SourceVertexInstanceIDs = InSource.GetPolygonVertexInstances(SourcePolygonID);
+		const FPolygonGroupID PolygonGroupID = InSource.GetPolygonPolygonGroup(SourcePolygonID);
 
 		TArray<FVertexInstanceID> ContourVertexInstances;
 		CreateContour(SourceVertexInstanceIDs, ContourVertexInstances);
 
 		// Insert a polygon into the mesh
-		const FPolygonID TargetPolygonID = InTarget.CreatePolygon(SourcePolygon.PolygonGroupID, ContourVertexInstances);
-		//Triangulate the polygon
-		FMeshPolygon& Polygon = InTarget.GetPolygon(TargetPolygonID);
-		InTarget.ComputePolygonTriangulation(TargetPolygonID, Polygon.Triangles);
+		const FPolygonID TargetPolygonID = InTarget.CreatePolygon(PolygonGroupID, ContourVertexInstances);
 	}
 }
 
diff --git a/Engine/Source/Developer/MeshSimplifier/Private/QuadricMeshReduction.cpp b/Engine/Source/Developer/MeshSimplifier/Private/QuadricMeshReduction.cpp
--- a/Engine/Source/Developer/MeshSimplifier/Private/QuadricMeshReduction.cpp
+++ b/Engine/Source/Developer/MeshSimplifier/Private/QuadricMeshReduction.cpp
@@ -211,11 +211,7 @@ public:
 
 		TMap< int32, int32 > VertsMap;
 
-		int32 NumFaces = 0;
-		for (const FPolygonID PolygonID : InMesh.Polygons().GetElementIDs())
-		{
-			NumFaces += InMesh.GetPolygonTriangles(PolygonID).Num();
-		}
+		int32 NumFaces = InMesh.Triangles().Num();
 		int32 NumWedges = NumFaces * 3;
 		FStaticMeshDescriptionConstAttributeGetter InMeshAttribute(&InMesh);
 		TVertexAttributesConstRef<FVector> InVertexPositions = InMeshAttribute.GetPositions();
@@ -233,15 +229,15 @@ public:
 		{
 			const FPolygonGroupID PolygonGroupID = InMesh.GetPolygonPolygonGroup(PolygonID);
 
-			const TArray<FMeshTriangle>& PolygonTriangles = InMesh.GetPolygonTriangles(PolygonID);
-			for (int32 TriangleIndex = 0; TriangleIndex < PolygonTriangles.Num(); ++TriangleIndex)
+			const TArray<FTriangleID>& TriangleIDs = InMesh.GetPolygonTriangleIDs(PolygonID);
+			for (int32 TriangleIndex = 0; TriangleIndex < TriangleIDs.Num(); ++TriangleIndex)
 			{
-				const FMeshTriangle& Triangle = PolygonTriangles[TriangleIndex];
+				const FTriangleID TriangleID = TriangleIDs[TriangleIndex];
 
 				FVector CornerPositions[3];
 				for (int32 TriVert = 0; TriVert < 3; ++TriVert)
 				{
-					const FVertexInstanceID VertexInstanceID = Triangle.GetVertexInstanceID(TriVert);
+					const FVertexInstanceID VertexInstanceID = InMesh.GetTriangleVertexInstance(TriangleID, TriVert);
 					const FVertexID TmpVertexID = InMesh.GetVertexInstanceVertex(VertexInstanceID);
 					const FVertexID VertexID = bWeldVertices ? VertexIDRemap[TmpVertexID] : TmpVertexID;
 					CornerPositions[TriVert] = InVertexPositions[VertexID];
@@ -259,7 +255,7 @@ public:
 				int32 VertexIndices[3];
 				for (int32 TriVert = 0; TriVert < 3; ++TriVert, ++WedgeIndex)
 				{
-					const FVertexInstanceID VertexInstanceID = Triangle.GetVertexInstanceID(TriVert);
+					const FVertexInstanceID VertexInstanceID = InMesh.GetTriangleVertexInstance(TriangleID, TriVert);
 					const int32 VertexInstanceValue = VertexInstanceID.GetValue();
 					const FVector& VertexPosition = CornerPositions[TriVert];
 
@@ -521,13 +517,6 @@ public:
 				{
 					// @todo: set NewEdgeID edge hardness?
 				}
-				const int32 NewTriangleIndex = OutReducedMesh.GetPolygonTriangles(NewPolygonID).AddDefaulted();
-				FMeshTriangle& NewTriangle = OutReducedMesh.GetPolygonTriangles(NewPolygonID)[NewTriangleIndex];
-				for (int32 TriangleVertexIndex = 0; TriangleVertexIndex < 3; ++TriangleVertexIndex)
-				{
-					const FVertexInstanceID VertexInstanceID = CornerInstanceIDs[TriangleVertexIndex];
-					NewTriangle.SetVertexInstanceID(TriangleVertexIndex, VertexInstanceID);
-				}
 			}
 			Verts.Empty();
 			Indexes.Empty();
@@ -536,8 +525,7 @@ public:
 			TArray<FPolygonGroupID> ToDeletePolygonGroupIDs;
 			for (const FPolygonGroupID& PolygonGroupID : OutReducedMesh.PolygonGroups().GetElementIDs())
 			{
-				FMeshPolygonGroup& PolygonGroup = OutReducedMesh.GetPolygonGroup(PolygonGroupID);
-				if (PolygonGroup.Polygons.Num() == 0)
+				if (OutReducedMesh.GetPolygonGroupPolygons(PolygonGroupID).Num() == 0)
 				{
 					ToDeletePolygonGroupIDs.Add(PolygonGroupID);
 				}
diff --git a/Engine/Source/Editor/UnrealEd/Private/EditorExporters.cpp b/Engine/Source/Editor/UnrealEd/Private/EditorExporters.cpp
--- a/Engine/Source/Editor/UnrealEd/Private/EditorExporters.cpp
+++ b/Engine/Source/Editor/UnrealEd/Private/EditorExporters.cpp
@@ -1751,11 +1751,7 @@ bool UStaticMeshExporterOBJ::ExportText(const FExportObjectInnerContext* Context
 		FMeshDescription* MeshDescription = StaticMesh->GetMeshDescription(0);
 		if (MeshDescription != nullptr)
 		{
-			uint32 TriangleCount = 0;
-			for (const FPolygonID PolygonID : MeshDescription->Polygons().GetElementIDs())
-			{
-				TriangleCount += MeshDescription->GetPolygonTriangles(PolygonID).Num();
-			}
+			uint32 TriangleCount = MeshDescription->Triangles().Num();
 			if (Count == TriangleCount)
 			{
 				SmoothingMasks.AddZeroed(TriangleCount);
diff --git a/Engine/Source/Editor/UnrealEd/Private/Fbx/FbxMainExport.cpp b/Engine/Source/Editor/UnrealEd/Private/Fbx/FbxMainExport.cpp
--- a/Engine/Source/Editor/UnrealEd/Private/Fbx/FbxMainExport.cpp
+++ b/Engine/Source/Editor/UnrealEd/Private/Fbx/FbxMainExport.cpp
@@ -1102,10 +1102,6 @@ void FFbxExporter::ExportBSP( UModel* Model, bool bSelectedOnly )
 					EdgeCreaseSharpnesses[EdgeID] = 0.0f;
 				}
 
-				//Triangulate the polygon
-				FMeshPolygon& Polygon = Mesh.GetPolygon(NewPolygonID);
-				Mesh.ComputePolygonTriangulation(NewPolygonID, Polygon.Triangles);
-
 				//Add to the smoothGroup array so we can compute hard edge later
 				ExportData->SmoothGroups.Add((1 << (Node.iSurf % 32)));
 
diff --git a/Engine/Source/Editor/UnrealEd/Private/Fbx/FbxStaticMeshImport.cpp b/Engine/Source/Editor/UnrealEd/Private/Fbx/FbxStaticMeshImport.cpp
--- a/Engine/Source/Editor/UnrealEd/Private/Fbx/FbxStaticMeshImport.cpp
+++ b/Engine/Source/Editor/UnrealEd/Private/Fbx/FbxStaticMeshImport.cpp
@@ -904,9 +904,6 @@ bool UnFbx::FFbxImporter::BuildStaticMeshFromGeometry(FbxNode* Node, UStaticMesh
 		TArray<FEdgeID> NewEdgeIDs;
 		const FPolygonID NewPolygonID = MeshDescription->CreatePolygon(PolygonGroupID, CornerInstanceIDs, &NewEdgeIDs);
 		check(NewEdgeIDs.Num() == 0);
-		//Triangulate the polygon
-		FMeshPolygon& Polygon = MeshDescription->GetPolygon(NewPolygonID);
-		MeshDescription->ComputePolygonTriangulation(NewPolygonID, Polygon.Triangles);
 	}
 	//Call this after all GetMeshEdgeIndexForPolygon call this is for optimization purpose.
 	Mesh->EndGetMeshEdgeIndexForPolygon();
diff --git a/Engine/Source/Editor/UnrealEd/Private/StaticMeshEdit.cpp b/Engine/Source/Editor/UnrealEd/Private/StaticMeshEdit.cpp
--- a/Engine/Source/Editor/UnrealEd/Private/StaticMeshEdit.cpp
+++ b/Engine/Source/Editor/UnrealEd/Private/StaticMeshEdit.cpp
@@ -787,13 +787,6 @@ void GetBrushMesh(ABrush* Brush, UModel* Model, FMeshDescription& MeshDescriptio
 				//All edge are hard for BSP
 				EdgeHardnesses[NewEdgeID] = true;
 			}
-			int32 NewTriangleIndex = MeshDescription.GetPolygonTriangles(NewPolygonID).AddDefaulted();
-			FMeshTriangle& NewTriangle = MeshDescription.GetPolygonTriangles(NewPolygonID)[NewTriangleIndex];
-			for (int32 TriangleVertexIndex = 0; TriangleVertexIndex < 3; ++TriangleVertexIndex)
-			{
-				const FVertexInstanceID &VertexInstanceID = VertexInstanceIDs[TriangleVertexIndex];
-				NewTriangle.SetVertexInstanceID(TriangleVertexIndex, VertexInstanceID);
-			}
 		}
 	}
 }
diff --git a/Engine/Source/Runtime/Core/Public/UObject/EditorObjectVersion.h b/Engine/Source/Runtime/Core/Public/UObject/EditorObjectVersion.h
--- a/Engine/Source/Runtime/Core/Public/UObject/EditorObjectVersion.h
+++ b/Engine/Source/Runtime/Core/Public/UObject/EditorObjectVersion.h
@@ -79,6 +79,8 @@ struct CORE_API FEditorObjectVersion
 		ScrollBarThicknessChange,
 		// Deprecated LandscapeHoleMaterial
 		RemoveLandscapeHoleMaterial,
+		// MeshDescription defined by triangles instead of arbitrary polygons
+		MeshDescriptionTriangles,
 		// -----<new versions can be added above this line>-------------------------------------------------
 		VersionPlusOne,
 		LatestVersion = VersionPlusOne - 1
diff --git a/Engine/Source/Runtime/Engine/Private/StaticMesh.cpp b/Engine/Source/Runtime/Engine/Private/StaticMesh.cpp
--- a/Engine/Source/Runtime/Engine/Private/StaticMesh.cpp
+++ b/Engine/Source/Runtime/Engine/Private/StaticMesh.cpp
@@ -2014,7 +2014,7 @@ static void SerializeBuildSettingsForDDC(FArchive& Ar, FMeshBuildSettings& Build
 // differences, etc.) replace the version GUID below with a new one.
 // In case of merge conflicts with DDC versions, you MUST generate a new GUID
 // and set this new GUID as the version.                                       
-#define STATICMESH_DERIVEDDATA_VER TEXT("B5FB810437E4428D9CC6367AE010BEEC")
+#define STATICMESH_DERIVEDDATA_VER TEXT("7D08EB504D3A42A4AAFAE139899C44E7")
 
 static const FString& GetStaticMeshDerivedDataVersion()
 {
@@ -3880,7 +3880,7 @@ void UStaticMesh::FixupMaterialSlotName()
 // differences, etc.) replace the version GUID below with a new one.
 // In case of merge conflicts with DDC versions, you MUST generate a new GUID
 // and set this new GUID as the version.                                       
-#define MESHDATAKEY_STATICMESH_DERIVEDDATA_VER TEXT("6E342F4997294EBEA9FB3ED1E8CF5D3B")
+#define MESHDATAKEY_STATICMESH_DERIVEDDATA_VER TEXT("A3E0B7AD760A496A8C56C261B5FE9BF9")
 
 static const FString& GetMeshDataKeyStaticMeshDerivedDataVersion()
 {
diff --git a/Engine/Source/Runtime/Landscape/Private/LandscapeEdit.cpp b/Engine/Source/Runtime/Landscape/Private/LandscapeEdit.cpp
--- a/Engine/Source/Runtime/Landscape/Private/LandscapeEdit.cpp
+++ b/Engine/Source/Runtime/Landscape/Private/LandscapeEdit.cpp
@@ -3160,9 +3160,6 @@ bool ALandscapeProxy::ExportToRawMesh(int32 InExportLOD, FMeshDescription& OutRa
 									EdgeHardnesses[NewEdgeID] = false;
 									EdgeCreaseSharpnesses[NewEdgeID] = 0.0f;
 								}
-								//Triangulate the polygon
-								FMeshPolygon& Polygon = OutRawMesh.GetPolygon(NewPolygonID);
-								OutRawMesh.ComputePolygonTriangulation(NewPolygonID, Polygon.Triangles);
 							};
 							AddTriangle(0);
 							AddTriangle(3);
