Add an option to allow compute normals to use weight by area and angle.
#rn Add an fbx option to allow compute normals to use weight by area and angle for both staticmesh and skeletalmesh. Add also a staticmesh build options for the feature.

diff --git a/Engine/Config/BaseEditorPerProjectUserSettings.ini b/Engine/Config/BaseEditorPerProjectUserSettings.ini
--- a/Engine/Config/BaseEditorPerProjectUserSettings.ini
+++ b/Engine/Config/BaseEditorPerProjectUserSettings.ini
@@ -653,6 +653,7 @@ bReorderMaterialToFbxOrder=True
 bImportMeshLODs=False
 NormalImportMethod=FBXNIM_ComputeNormals
 NormalGenerationMethod=EFBXNormalGenerationMethod::MikkTSpace
+bComputeWeightedNormals=True
 
 [/Script/UnrealEd.FbxStaticMeshImportData]
 StaticMeshLODGroup=""
diff --git a/Engine/Plugins/Editor/GLTFImporter/Source/GLTFImporter/Private/GLTFStaticMeshFactory.cpp b/Engine/Plugins/Editor/GLTFImporter/Source/GLTFImporter/Private/GLTFStaticMeshFactory.cpp
--- a/Engine/Plugins/Editor/GLTFImporter/Source/GLTFImporter/Private/GLTFStaticMeshFactory.cpp
+++ b/Engine/Plugins/Editor/GLTFImporter/Source/GLTFImporter/Private/GLTFStaticMeshFactory.cpp
@@ -333,6 +333,7 @@ namespace GLTF
 		Settings.bRecomputeNormals  = false;
 		Settings.bRecomputeTangents = !bMeshHasTagents;
 		Settings.bUseMikkTSpace     = true;  // glTF spec defines that MikkTSpace algorithms should be used when tangents aren't defined
+		Settings.bComputeWeightedNormals = true;
 
 		Settings.bRemoveDegenerates        = false;
 		Settings.bBuildAdjacencyBuffer     = false;
diff --git a/Engine/Plugins/Experimental/DataPrepEditor/Source/DataPrepLibraries/Private/DataPrepOperationsLibrary.cpp b/Engine/Plugins/Experimental/DataPrepEditor/Source/DataPrepLibraries/Private/DataPrepOperationsLibrary.cpp
--- a/Engine/Plugins/Experimental/DataPrepEditor/Source/DataPrepLibraries/Private/DataPrepOperationsLibrary.cpp
+++ b/Engine/Plugins/Experimental/DataPrepEditor/Source/DataPrepLibraries/Private/DataPrepOperationsLibrary.cpp
@@ -157,6 +157,7 @@ namespace DataprepOperationsLibraryUtil
 				SourceModel.BuildSettings.bRecomputeTangents = false;
 				SourceModel.BuildSettings.bBuildAdjacencyBuffer = false;
 				SourceModel.BuildSettings.bBuildReversedIndexBuffer = false;
+				SourceModel.BuildSettings.bComputeWeightedNormals = false;
 			}
 
 			return BuildSettingsBackup;
@@ -184,6 +185,7 @@ namespace DataprepOperationsLibraryUtil
 					BuildSettings.bRecomputeTangents = CachedBuildSettings.bRecomputeTangents;
 					BuildSettings.bBuildAdjacencyBuffer = CachedBuildSettings.bBuildAdjacencyBuffer;
 					BuildSettings.bBuildReversedIndexBuffer = CachedBuildSettings.bBuildReversedIndexBuffer;
+					BuildSettings.bComputeWeightedNormals = CachedBuildSettings.bComputeWeightedNormals;
 				}
 			}
 		}
@@ -237,6 +239,7 @@ namespace DataprepOperationsLibraryUtil
 		SourceModel.BuildSettings.bRecomputeTangents = false;
 		SourceModel.BuildSettings.bBuildAdjacencyBuffer = false;
 		SourceModel.BuildSettings.bBuildReversedIndexBuffer = false;
+		SourceModel.BuildSettings.bComputeWeightedNormals = false;
 
 		FMeshDescription* StaticMeshDescription = SourceModel.MeshDescription.Get();
 		check( StaticMeshDescription );
diff --git a/Engine/Plugins/Runtime/ApexDestruction/Source/ApexDestruction/Private/ApexDestructibleAssetImport.cpp b/Engine/Plugins/Runtime/ApexDestruction/Source/ApexDestruction/Private/ApexDestructibleAssetImport.cpp
--- a/Engine/Plugins/Runtime/ApexDestruction/Source/ApexDestruction/Private/ApexDestructibleAssetImport.cpp
+++ b/Engine/Plugins/Runtime/ApexDestruction/Source/ApexDestruction/Private/ApexDestructibleAssetImport.cpp
@@ -912,6 +912,8 @@ bool SetApexDestructibleAsset(UDestructibleMesh& DestructibleMesh, apex::Destruc
 		IMeshUtilities::MeshBuildOptions BuildOptions;
 		BuildOptions.bComputeNormals = !bHaveNormals;
 		BuildOptions.bComputeTangents = !bHaveTangents;
+		BuildOptions.bUseMikkTSpace = true;
+		BuildOptions.bComputeWeightedNormals = true;
 
 		// Create actual rendering data.
 		if (!MeshUtilities.BuildSkeletalMesh(DestructibleMeshResource.LODModels[0], DestructibleMesh.RefSkeleton, LODInfluences,LODWedges,LODFaces,LODPoints,LODPointToRawMap,BuildOptions))
diff --git a/Engine/Plugins/Tests/EditorTests/Source/EditorTests/Private/UnrealEd/FbxAutomationTests.cpp b/Engine/Plugins/Tests/EditorTests/Source/EditorTests/Private/UnrealEd/FbxAutomationTests.cpp
--- a/Engine/Plugins/Tests/EditorTests/Source/EditorTests/Private/UnrealEd/FbxAutomationTests.cpp
+++ b/Engine/Plugins/Tests/EditorTests/Source/EditorTests/Private/UnrealEd/FbxAutomationTests.cpp
@@ -296,6 +296,7 @@ bool FFbxImportAssetsAutomationTest::RunTest(const FString& Parameters)
 					ImportData->bImportMeshLODs = TestPlan->ImportUI->StaticMeshImportData->bImportMeshLODs;
 					ImportData->NormalImportMethod = TestPlan->ImportUI->StaticMeshImportData->NormalImportMethod;
 					ImportData->NormalGenerationMethod = TestPlan->ImportUI->StaticMeshImportData->NormalGenerationMethod;
+					ImportData->bComputeWeightedNormals = TestPlan->ImportUI->StaticMeshImportData->bComputeWeightedNormals;
 					//Copy UFbxAssetImportData
 					ImportData->ImportTranslation = TestPlan->ImportUI->StaticMeshImportData->ImportTranslation;
 					ImportData->ImportRotation = TestPlan->ImportUI->StaticMeshImportData->ImportRotation;
@@ -401,6 +402,7 @@ bool FFbxImportAssetsAutomationTest::RunTest(const FString& Parameters)
 					ImportData->bImportMeshLODs = TestPlan->ImportUI->StaticMeshImportData->bImportMeshLODs;
 					ImportData->NormalImportMethod = TestPlan->ImportUI->StaticMeshImportData->NormalImportMethod;
 					ImportData->NormalGenerationMethod = TestPlan->ImportUI->StaticMeshImportData->NormalGenerationMethod;
+					ImportData->bComputeWeightedNormals = TestPlan->ImportUI->StaticMeshImportData->bComputeWeightedNormals;
 					//Copy UFbxAssetImportData
 					ImportData->ImportTranslation = TestPlan->ImportUI->StaticMeshImportData->ImportTranslation;
 					ImportData->ImportRotation = TestPlan->ImportUI->StaticMeshImportData->ImportRotation;
diff --git a/Engine/Source/Developer/HierarchicalLODUtilities/Private/HierarchicalLODUtilities.cpp b/Engine/Source/Developer/HierarchicalLODUtilities/Private/HierarchicalLODUtilities.cpp
--- a/Engine/Source/Developer/HierarchicalLODUtilities/Private/HierarchicalLODUtilities.cpp
+++ b/Engine/Source/Developer/HierarchicalLODUtilities/Private/HierarchicalLODUtilities.cpp
@@ -292,6 +292,7 @@ UStaticMesh* CreateImposterStaticMesh(UStaticMeshComponent* InComponent, UMateri
 		/*Don't allow the engine to recalculate normals*/
 		SrcModel.BuildSettings.bRecomputeNormals = false;
 		SrcModel.BuildSettings.bRecomputeTangents = false;
+		SrcModel.BuildSettings.bComputeWeightedNormals = true;
 		SrcModel.BuildSettings.bRemoveDegenerates = true;
 		SrcModel.BuildSettings.bUseHighPrecisionTangentBasis = false;
 		SrcModel.BuildSettings.bUseFullPrecisionUVs = false;
diff --git a/Engine/Source/Developer/MeshBuilder/Private/MeshDescriptionHelper.cpp b/Engine/Source/Developer/MeshBuilder/Private/MeshDescriptionHelper.cpp
--- a/Engine/Source/Developer/MeshBuilder/Private/MeshDescriptionHelper.cpp
+++ b/Engine/Source/Developer/MeshBuilder/Private/MeshDescriptionHelper.cpp
@@ -61,6 +61,11 @@ void FMeshDescriptionHelper::GetRenderMeshDescription(UObject* Owner, const FMes
 			// If removing degenerate triangles, ignore them when computing tangents.
 			TangentOptions |= FMeshDescriptionOperations::ETangentOptions::IgnoreDegenerateTriangles;
 		}
+		if (BuildSettings->bComputeWeightedNormals)
+		{
+			// Specify we want to compute weighted normals.
+			TangentOptions |= FMeshDescriptionOperations::ETangentOptions::UseWeightedAreaAndAngle;
+		}
 		
 		//Keep the original mesh description NTBs if we do not rebuild the normals or tangents.
 		bool bComputeTangentLegacy = !BuildSettings->bUseMikkTSpace && (BuildSettings->bRecomputeNormals || BuildSettings->bRecomputeTangents);
diff --git a/Engine/Source/Developer/MeshDescriptionOperations/Private/MeshDescriptionOperations.cpp b/Engine/Source/Developer/MeshDescriptionOperations/Private/MeshDescriptionOperations.cpp
--- a/Engine/Source/Developer/MeshDescriptionOperations/Private/MeshDescriptionOperations.cpp
+++ b/Engine/Source/Developer/MeshDescriptionOperations/Private/MeshDescriptionOperations.cpp
@@ -11,6 +11,7 @@
 #include "mikktspace.h"
 #include "UVMapSettings.h"
 #include "Modules/ModuleManager.h"
+#include "MeshUtilitiesCommon.h"
 
 IMPLEMENT_MODULE(FDefaultModuleImpl, MeshDescriptionOperations)
 
@@ -24,13 +25,13 @@ struct FVertexInfo
 {
 	FVertexInfo()
 	{
-		PolygonID = FPolygonID::Invalid;
+		TriangleID = FTriangleID::Invalid;
 		VertexInstanceID = FVertexInstanceID::Invalid;
 		UVs = FVector2D(0.0f, 0.0f);
 		EdgeIDs.Reserve(2);//Most of the time a edge has two triangles
 	}
 
-	FPolygonID PolygonID;
+	FTriangleID TriangleID;
 	FVertexInstanceID VertexInstanceID;
 	FVector2D UVs;
 	TArray<FEdgeID> EdgeIDs;
@@ -745,10 +746,11 @@ void FMeshDescriptionOperations::AppendMeshDescription(const FMeshDescription& S
 
 //////////////////////////////////////////////////////////////////////////
 // Normals tangents and Bi-normals
-void FMeshDescriptionOperations::RecomputeNormalsAndTangentsIfNeeded(FMeshDescription& MeshDescription, ETangentOptions TangentOptions, bool bUseMikkTSpace, bool bForceRecomputeNormals, bool bForceRecomputeTangents)
+void FMeshDescriptionOperations::RecomputeNormalsAndTangentsIfNeeded(FMeshDescription& MeshDescription, ETangentOptions TangentOptions, bool bForceRecomputeNormals, bool bForceRecomputeTangents)
 {
 	bool bRecomputeNormals = bForceRecomputeNormals;
 	bool bRecomputeTangents = bForceRecomputeTangents;
+	bool bUseMikkTSpace = ((TangentOptions & ETangentOptions::UseMikkTSpace) != 0);
 	TVertexInstanceAttributesRef<FVector> VertexInstanceNormals = MeshDescription.VertexInstanceAttributes().GetAttributesRef<FVector>(MeshAttribute::VertexInstance::Normal);
 	TVertexInstanceAttributesRef<FVector> VertexInstanceTangents = MeshDescription.VertexInstanceAttributes().GetAttributesRef<FVector>(MeshAttribute::VertexInstance::Tangent);
 	if (!bRecomputeNormals || !bRecomputeTangents)
@@ -927,6 +929,17 @@ void FMeshDescriptionOperations::CreateNormals(FMeshDescription& MeshDescription
 	// the angle it makes with the vertex being calculated. This means that triangulated faces whose
 	// internal edge meets the vertex doesn't get undue extra weight.
 
+	//Make sure the meshdescription is triangulate
+	if (MeshDescription.Triangles().Num() < MeshDescription.Polygons().Num())
+	{
+		//Triangulate the mesh, we compute the normals on triangle not on polygon.
+		MeshDescription.TriangulateMesh();
+	}
+
+	bool bComputeWeightedNormals = (TangentOptions & ETangentOptions::UseWeightedAreaAndAngle) != 0;
+
+	TVertexAttributesConstRef<FVector> VertexPositions = MeshDescription.VertexAttributes().GetAttributesRef<FVector>(MeshAttribute::Vertex::Position);
+
 	TVertexInstanceAttributesConstRef<FVector2D> VertexUVs = MeshDescription.VertexInstanceAttributes().GetAttributesRef<FVector2D>(MeshAttribute::VertexInstance::TextureCoordinate);
 	TVertexInstanceAttributesRef<FVector> VertexNormals = MeshDescription.VertexInstanceAttributes().GetAttributesRef<FVector>(MeshAttribute::VertexInstance::Normal);
 	TVertexInstanceAttributesRef<FVector> VertexTangents = MeshDescription.VertexInstanceAttributes().GetAttributesRef<FVector>(MeshAttribute::VertexInstance::Tangent);
@@ -939,7 +952,67 @@ void FMeshDescriptionOperations::CreateNormals(FMeshDescription& MeshDescription
 	check(PolygonTangents.IsValid());
 	check(PolygonBinormals.IsValid());
 
-	TMap<FPolygonID, FVertexInfo> VertexInfoMap;
+	struct FTriangleCornerData
+	{
+		FVertexInstanceID VertexInstanceID;
+		float CornerAngle;
+	};
+	struct FTriangleData
+	{
+	public:
+		//The area of the triangle
+		float Area;
+
+		//Set the corner angle data for a FVertexInstanceID
+		void SetCornerAngleData(FVertexInstanceID VertexInstanceID, float CornerAngle, int32 CornerIndex)
+		{
+			CornerAngleDatas[CornerIndex].VertexInstanceID = VertexInstanceID;
+			CornerAngleDatas[CornerIndex].CornerAngle = CornerAngle;
+		}
+
+		//Get the angle for the FVertexInstanceID
+		float GetCornerAngle(FVertexInstanceID VertexInstanceID)
+		{
+			for (int32 CornerIndex = 0; CornerIndex < 3; ++CornerIndex)
+			{
+				if (CornerAngleDatas[CornerIndex].VertexInstanceID == VertexInstanceID)
+				{
+					return CornerAngleDatas[CornerIndex].CornerAngle;
+				}
+			}
+			
+			//We should always found a valid VertexInstanceID
+			check(false);
+			return 0.0f;
+		}
+	private:
+		//The data for each corner
+		FTriangleCornerData CornerAngleDatas[3];
+	};
+
+	//Compute the weight (area and angle) for each triangles
+	TMap<FTriangleID, FTriangleData> TriangleDatas;
+	if(bComputeWeightedNormals)
+	{
+		TriangleDatas.Reserve(MeshDescription.Triangles().Num());
+		
+		for (FTriangleID TriangleID : MeshDescription.Triangles().GetElementIDs())
+		{
+			TArrayView<const FVertexInstanceID> VertexInstanceIDs = MeshDescription.GetTriangleVertexInstances(TriangleID);
+			//Triangle should use 3 vertex instances
+			check(VertexInstanceIDs.Num() == 3);
+			const FVector& PointA = VertexPositions[MeshDescription.GetVertexInstanceVertex(VertexInstanceIDs[0])];
+			const FVector& PointB = VertexPositions[MeshDescription.GetVertexInstanceVertex(VertexInstanceIDs[1])];
+			const FVector& PointC = VertexPositions[MeshDescription.GetVertexInstanceVertex(VertexInstanceIDs[2])];
+			FTriangleData& TriangleData = TriangleDatas.FindOrAdd(TriangleID);
+			TriangleData.Area = TriangleUtilities::ComputeTriangleArea(PointA, PointB, PointC);
+			TriangleData.SetCornerAngleData(VertexInstanceIDs[0], TriangleUtilities::ComputeTriangleCornerAngle(PointA, PointB, PointC), 0);
+			TriangleData.SetCornerAngleData(VertexInstanceIDs[1], TriangleUtilities::ComputeTriangleCornerAngle(PointB, PointC, PointA), 1);
+			TriangleData.SetCornerAngleData(VertexInstanceIDs[2], TriangleUtilities::ComputeTriangleCornerAngle(PointC, PointA, PointB), 2);
+		}
+	}
+
+	TMap<FTriangleID, FVertexInfo> VertexInfoMap;
 	VertexInfoMap.Reserve(20);
 	//Iterate all vertex to compute normals for all vertex instance
 	for (const FVertexID VertexID : MeshDescription.Vertices().GetElementIDs())
@@ -950,14 +1023,14 @@ void FMeshDescriptionOperations::CreateNormals(FMeshDescription& MeshDescription
 		//Fill the VertexInfoMap
 		for (const FEdgeID EdgeID : MeshDescription.GetVertexConnectedEdges(VertexID))
 		{
-			for (const FPolygonID PolygonID : MeshDescription.GetEdgeConnectedPolygons(EdgeID))
+			for (const FTriangleID TriangleID : MeshDescription.GetEdgeConnectedTriangles(EdgeID))
 			{
-				FVertexInfo& VertexInfo = VertexInfoMap.FindOrAdd(PolygonID);
+				FVertexInfo& VertexInfo = VertexInfoMap.FindOrAdd(TriangleID);
 				int32 EdgeIndex = VertexInfo.EdgeIDs.AddUnique(EdgeID);
-				if (VertexInfo.PolygonID == FPolygonID::Invalid)
+				if (VertexInfo.TriangleID == FTriangleID::Invalid)
 				{
-					VertexInfo.PolygonID = PolygonID;
-					for (const FVertexInstanceID VertexInstanceID : MeshDescription.GetPolygonVertexInstances(PolygonID))
+					VertexInfo.TriangleID = TriangleID;
+					for (const FVertexInstanceID VertexInstanceID : MeshDescription.GetTriangleVertexInstances(TriangleID))
 					{
 						if (MeshDescription.GetVertexInstanceVertex(VertexInstanceID) == VertexID)
 						{
@@ -989,25 +1062,25 @@ void FMeshDescriptionOperations::CreateNormals(FMeshDescription& MeshDescription
 		}
 
 		//Build all group by recursively traverse all polygon connected to the vertex
-		TArray<TArray<FPolygonID>> Groups;
-		TArray<FPolygonID> ConsumedPolygon;
+		TArray<TArray<FTriangleID>> Groups;
+		TArray<FTriangleID> ConsumedTriangle;
 		for (auto Kvp : VertexInfoMap)
 		{
-			if (ConsumedPolygon.Contains(Kvp.Key))
+			if (ConsumedTriangle.Contains(Kvp.Key))
 			{
 				continue;
 			}
 
 			int32 CurrentGroupIndex = Groups.AddZeroed();
-			TArray<FPolygonID>& CurrentGroup = Groups[CurrentGroupIndex];
-			TArray<FPolygonID> PolygonQueue;
+			TArray<FTriangleID>& CurrentGroup = Groups[CurrentGroupIndex];
+			TArray<FTriangleID> PolygonQueue;
 			PolygonQueue.Add(Kvp.Key); //Use a queue to avoid recursive function
 			while (PolygonQueue.Num() > 0)
 			{
-				FPolygonID CurrentPolygonID = PolygonQueue.Pop(true);
+				FTriangleID CurrentPolygonID = PolygonQueue.Pop(true);
 				FVertexInfo& CurrentVertexInfo = VertexInfoMap.FindOrAdd(CurrentPolygonID);
-				CurrentGroup.AddUnique(CurrentVertexInfo.PolygonID);
-				ConsumedPolygon.AddUnique(CurrentVertexInfo.PolygonID);
+				CurrentGroup.AddUnique(CurrentVertexInfo.TriangleID);
+				ConsumedTriangle.AddUnique(CurrentVertexInfo.TriangleID);
 				const TEdgeAttributesRef<bool> EdgeHardnesses = MeshDescription.EdgeAttributes().GetAttributesRef<bool>(MeshAttribute::Edge::IsHard);
 				for (const FEdgeID EdgeID : CurrentVertexInfo.EdgeIDs)
 				{
@@ -1016,18 +1089,18 @@ void FMeshDescriptionOperations::CreateNormals(FMeshDescription& MeshDescription
 						//End of the group
 						continue;
 					}
-					for (const FPolygonID PolygonID : MeshDescription.GetEdgeConnectedPolygons(EdgeID))
+					for (const FTriangleID TriangleID : MeshDescription.GetEdgeConnectedTriangles(EdgeID))
 					{
-						if (PolygonID == CurrentVertexInfo.PolygonID)
+						if (TriangleID == CurrentVertexInfo.TriangleID)
 						{
 							continue;
 						}
 						//Add this polygon to the group
-						FVertexInfo& OtherVertexInfo = VertexInfoMap.FindOrAdd(PolygonID);
+						FVertexInfo& OtherVertexInfo = VertexInfoMap.FindOrAdd(TriangleID);
 						//Do not repeat polygons
-						if (!ConsumedPolygon.Contains(OtherVertexInfo.PolygonID))
+						if (!ConsumedTriangle.Contains(OtherVertexInfo.TriangleID))
 						{
-							PolygonQueue.Add(PolygonID);
+							PolygonQueue.Add(TriangleID);
 						}
 					}
 				}
@@ -1035,8 +1108,8 @@ void FMeshDescriptionOperations::CreateNormals(FMeshDescription& MeshDescription
 		}
 
 		//Smooth every connected group
-		ConsumedPolygon.Reset();
-		for (const TArray<FPolygonID>& Group : Groups)
+		ConsumedTriangle.Reset();
+		for (const TArray<FTriangleID>& Group : Groups)
 		{
 			//Compute tangents data
 			TMap<FVector2D, FVector,FDefaultSetAllocator, FNTBGroupKeyFuncs> GroupTangent;
@@ -1044,21 +1117,31 @@ void FMeshDescriptionOperations::CreateNormals(FMeshDescription& MeshDescription
 
 			TArray<FVertexInstanceID> VertexInstanceInGroup;
 			FVector GroupNormal(0.0f);
-			for (const FPolygonID PolygonID : Group)
+			for (const FTriangleID TriangleID : Group)
 			{
-				FVector PolyNormal = PolygonNormals[PolygonID];
-				FVector PolyTangent = PolygonTangents[PolygonID];
-				FVector PolyBinormal = PolygonBinormals[PolygonID];
+				FPolygonID PolygonID = MeshDescription.GetTrianglePolygon(TriangleID);
+				FVertexInfo& CurrentVertexInfo = VertexInfoMap.FindOrAdd(TriangleID);
+				float CornerWeight = 1.0f;
+				
+				if (bComputeWeightedNormals)
+				{
+					FTriangleData& TriangleData = TriangleDatas.FindChecked(TriangleID);
+					CornerWeight = TriangleData.Area * TriangleData.GetCornerAngle(CurrentVertexInfo.VertexInstanceID);
+				}
+
+				FVector PolyNormal = CornerWeight * PolygonNormals[PolygonID];
+				FVector PolyTangent = CornerWeight * PolygonTangents[PolygonID];
+				FVector PolyBinormal = CornerWeight * PolygonBinormals[PolygonID];
 
-				ConsumedPolygon.Add(PolygonID);
-				VertexInstanceInGroup.Add(VertexInfoMap[PolygonID].VertexInstanceID);
+				ConsumedTriangle.Add(TriangleID);
+				VertexInstanceInGroup.Add(CurrentVertexInfo.VertexInstanceID);
 				if (!PolyNormal.IsNearlyZero(SMALL_NUMBER) && !PolyNormal.ContainsNaN())
 				{
 					GroupNormal += PolyNormal;
 				}
 				if (bComputeTangent)
 				{
-					const FVector2D UVs = VertexInfoMap[PolygonID].UVs;
+					const FVector2D UVs = VertexInfoMap[TriangleID].UVs;
 					bool CreateGroup = (!GroupTangent.Contains(UVs));
 					FVector& GroupTangentValue = GroupTangent.FindOrAdd(UVs);
 					FVector& GroupBiNormalValue = GroupBiNormal.FindOrAdd(UVs);
diff --git a/Engine/Source/Developer/MeshDescriptionOperations/Public/MeshDescriptionOperations.h b/Engine/Source/Developer/MeshDescriptionOperations/Public/MeshDescriptionOperations.h
--- a/Engine/Source/Developer/MeshDescriptionOperations/Public/MeshDescriptionOperations.h
+++ b/Engine/Source/Developer/MeshDescriptionOperations/Public/MeshDescriptionOperations.h
@@ -29,9 +29,10 @@ public:
 	enum ETangentOptions
 	{
 		None = 0,
-		BlendOverlappingNormals = 0x1,
-		IgnoreDegenerateTriangles = 0x2,
-		UseMikkTSpace = 0x4,
+		BlendOverlappingNormals = 0x00000001,
+		IgnoreDegenerateTriangles = 0x00000002,
+		UseMikkTSpace = 0x00000004,
+		UseWeightedAreaAndAngle = 0x00000008, //Use surface area and angle as a ratio when computing normals
 	};
 
 	/** Convert this mesh description into the old FRawMesh format. */
@@ -57,7 +58,7 @@ public:
 	/*
 	 * Check if all normals and tangents are valid, if not recompute them
 	 */
-	static void RecomputeNormalsAndTangentsIfNeeded(FMeshDescription& MeshDescription, ETangentOptions TangentOptions, bool bUseMikkTSpace, bool bForceRecomputeNormals = false, bool bForceRecomputeTangents = false);
+	static void RecomputeNormalsAndTangentsIfNeeded(FMeshDescription& MeshDescription, ETangentOptions TangentOptions, bool bForceRecomputeNormals = false, bool bForceRecomputeTangents = false);
 
 	/**
 	 * Compute normal, tangent and Bi-Normal for every polygon in the mesh description. (this do not compute Vertex NTBs)
diff --git a/Engine/Source/Developer/MeshMergeUtilities/Private/MeshMergeHelpers.cpp b/Engine/Source/Developer/MeshMergeUtilities/Private/MeshMergeHelpers.cpp
--- a/Engine/Source/Developer/MeshMergeUtilities/Private/MeshMergeHelpers.cpp
+++ b/Engine/Source/Developer/MeshMergeUtilities/Private/MeshMergeHelpers.cpp
@@ -244,8 +244,12 @@ void FMeshMergeHelpers::RetrieveMesh(const UStaticMeshComponent* StaticMeshCompo
 		// If removing degenerate triangles, ignore them when computing tangents.
 		TangentOptions |= FMeshDescriptionOperations::ETangentOptions::IgnoreDegenerateTriangles;
 	}
+	if (BuildSettings.bUseMikkTSpace)
+	{
+		TangentOptions |= FMeshDescriptionOperations::ETangentOptions::UseMikkTSpace;
+	}
 	FMeshDescriptionOperations::CreatePolygonNTB(RawMesh, 0.0f);
-	FMeshDescriptionOperations::RecomputeNormalsAndTangentsIfNeeded(RawMesh, (FMeshDescriptionOperations::ETangentOptions)TangentOptions, BuildSettings.bUseMikkTSpace);
+	FMeshDescriptionOperations::RecomputeNormalsAndTangentsIfNeeded(RawMesh, (FMeshDescriptionOperations::ETangentOptions)TangentOptions);
 }
 
 void FMeshMergeHelpers::RetrieveMesh(USkeletalMeshComponent* SkeletalMeshComponent, int32 LODIndex, FMeshDescription& RawMesh, bool bPropagateVertexColours)
@@ -400,8 +404,12 @@ void FMeshMergeHelpers::RetrieveMesh(const UStaticMesh* StaticMesh, int32 LODInd
 		// If removing degenerate triangles, ignore them when computing tangents.
 		TangentOptions |= FMeshDescriptionOperations::ETangentOptions::IgnoreDegenerateTriangles;
 	}
+	if (BuildSettings.bUseMikkTSpace)
+	{
+		TangentOptions |= FMeshDescriptionOperations::ETangentOptions::UseMikkTSpace;
+	}
 	FMeshDescriptionOperations::CreatePolygonNTB(RawMesh, 0.0f);
-	FMeshDescriptionOperations::RecomputeNormalsAndTangentsIfNeeded(RawMesh, (FMeshDescriptionOperations::ETangentOptions)TangentOptions, BuildSettings.bUseMikkTSpace, (bImportedMesh && BuildSettings.bRecomputeNormals), (bImportedMesh && BuildSettings.bRecomputeTangents));
+	FMeshDescriptionOperations::RecomputeNormalsAndTangentsIfNeeded(RawMesh, (FMeshDescriptionOperations::ETangentOptions)TangentOptions, (bImportedMesh && BuildSettings.bRecomputeNormals), (bImportedMesh && BuildSettings.bRecomputeTangents));
 }
 
 void FMeshMergeHelpers::ExportStaticMeshLOD(const FStaticMeshLODResources& StaticMeshLOD, FMeshDescription& OutRawMesh, const TArray<FStaticMaterial>& Materials)
diff --git a/Engine/Source/Developer/MeshMergeUtilities/Private/ProxyGenerationProcessor.cpp b/Engine/Source/Developer/MeshMergeUtilities/Private/ProxyGenerationProcessor.cpp
--- a/Engine/Source/Developer/MeshMergeUtilities/Private/ProxyGenerationProcessor.cpp
+++ b/Engine/Source/Developer/MeshMergeUtilities/Private/ProxyGenerationProcessor.cpp
@@ -194,6 +194,7 @@ void FProxyGenerationProcessor::ProcessJob(const FGuid& JobGuid, FProxyGeneratio
 	/*Don't allow the engine to recalculate normals*/
 	SrcModel.BuildSettings.bRecomputeNormals = false;
 	SrcModel.BuildSettings.bRecomputeTangents = false;
+	SrcModel.BuildSettings.bComputeWeightedNormals = true;
 	SrcModel.BuildSettings.bRemoveDegenerates = true;
 	SrcModel.BuildSettings.bUseHighPrecisionTangentBasis = false;
 	SrcModel.BuildSettings.bUseFullPrecisionUVs = false;
diff --git a/Engine/Source/Developer/MeshUtilities/Private/MeshUtilities.cpp b/Engine/Source/Developer/MeshUtilities/Private/MeshUtilities.cpp
--- a/Engine/Source/Developer/MeshUtilities/Private/MeshUtilities.cpp
+++ b/Engine/Source/Developer/MeshUtilities/Private/MeshUtilities.cpp
@@ -3409,6 +3409,7 @@ public:
 	{
 		bool bBlendOverlappingNormals = true;
 		bool bIgnoreDegenerateTriangles = BuildData->BuildOptions.bRemoveDegenerateTriangles;
+		bool bComputeWeightedNormals = BuildData->BuildOptions.bComputeWeightedNormals;
 		
 		// Compute per-triangle tangents.
 		TArray<FVector> TriangleTangentX;
@@ -3586,8 +3587,8 @@ public:
 								FFanFace& NextFace = RelevantFacesForCorner[CornerIndex][NextFaceIndex];
 								if (!NextFace.bFilled) // && !NextFace.bBlendTangents)
 								{
-									if (NextFaceIndex != OtherFaceIdx)
-										//&& (BuildData->GetFaceSmoothingGroups(NextFace.FaceIndex) & BuildData->GetFaceSmoothingGroups(OtherFace.FaceIndex)))
+									if ((NextFaceIndex != OtherFaceIdx)
+										&& (BuildData->GetFaceSmoothingGroups(NextFace.FaceIndex) & BuildData->GetFaceSmoothingGroups(OtherFace.FaceIndex)))
 									{
 										int32 CommonVertices = 0;
 										int32 CommonTangentVertices = 0;
@@ -3668,14 +3669,26 @@ public:
 						if (RelevantFace.bFilled)
 						{
 							int32 OtherFaceIndex = RelevantFace.FaceIndex;
+							float CornerWeight = 1.0f;
+							if (bComputeWeightedNormals)
+							{
+
+								FVector OtherFacePoint[3] = { BuildData->GetVertexPosition(OtherFaceIndex, 0), BuildData->GetVertexPosition(OtherFaceIndex, 1), BuildData->GetVertexPosition(OtherFaceIndex, 2) };
+								float OtherFaceArea = TriangleUtilities::ComputeTriangleArea(OtherFacePoint[0], OtherFacePoint[1], OtherFacePoint[2]);
+								int32 OtherFaceCornerIndex = RelevantFace.LinkedVertexIndex;
+								float OtherFaceAngle = TriangleUtilities::ComputeTriangleCornerAngle(OtherFacePoint[OtherFaceCornerIndex], OtherFacePoint[(OtherFaceCornerIndex + 1) % 3], OtherFacePoint[(OtherFaceCornerIndex + 2) % 3]);
+								//Get the CornerWeight
+								CornerWeight = OtherFaceArea * OtherFaceAngle;
+							}
+
 							if (RelevantFace.bBlendTangents)
 							{
-								CornerTangentX[CornerIndex] += TriangleTangentX[OtherFaceIndex];
-								CornerTangentY[CornerIndex] += TriangleTangentY[OtherFaceIndex];
+								CornerTangentX[CornerIndex] += CornerWeight * TriangleTangentX[OtherFaceIndex];
+								CornerTangentY[CornerIndex] += CornerWeight * TriangleTangentY[OtherFaceIndex];
 							}
 							if (RelevantFace.bBlendNormals)
 							{
-								CornerTangentZ[CornerIndex] += TriangleTangentZ[OtherFaceIndex];
+								CornerTangentZ[CornerIndex] += CornerWeight * TriangleTangentZ[OtherFaceIndex];
 							}
 						}
 					}
@@ -3765,6 +3778,7 @@ public:
 	{
 		bool bBlendOverlappingNormals = true;
 		bool bIgnoreDegenerateTriangles = BuildData->BuildOptions.bRemoveDegenerateTriangles;
+		bool bComputeWeightedNormals = BuildData->BuildOptions.bComputeWeightedNormals;
 		
 		int32 NumFaces = BuildData->GetNumFaces();
 		int32 NumWedges = BuildData->GetNumWedges();
@@ -4003,9 +4017,20 @@ public:
 						if (RelevantFace.bFilled)
 						{
 							int32 OtherFaceIndex = RelevantFace.FaceIndex;
+							float CornerWeight = 1.0f;
+							if (bComputeWeightedNormals)
+							{
+								FVector OtherFacePoint[3] = { BuildData->GetVertexPosition(OtherFaceIndex, 0), BuildData->GetVertexPosition(OtherFaceIndex, 1), BuildData->GetVertexPosition(OtherFaceIndex, 2) };
+								float OtherFaceArea = TriangleUtilities::ComputeTriangleArea(OtherFacePoint[0], OtherFacePoint[1], OtherFacePoint[2]);
+								int32 OtherFaceCornerIndex = RelevantFace.LinkedVertexIndex;
+								float OtherFaceAngle = TriangleUtilities::ComputeTriangleCornerAngle(OtherFacePoint[OtherFaceCornerIndex], OtherFacePoint[(OtherFaceCornerIndex + 1) % 3], OtherFacePoint[(OtherFaceCornerIndex + 2) % 3]);
+								//Get the CornerWeight
+								CornerWeight = OtherFaceArea * OtherFaceAngle;
+							}
+							
 							if (RelevantFace.bBlendNormals)
 							{
-								CornerNormal[CornerIndex] += TriangleTangentZ[OtherFaceIndex];
+								CornerNormal[CornerIndex] += CornerWeight * TriangleTangentZ[OtherFaceIndex];
 							}
 						}
 					}
@@ -4383,7 +4408,7 @@ bool FMeshUtilities::BuildSkeletalMesh(FSkeletalMeshLODModel& LODModel, const FR
 	// Temporarily supporting both import paths
 	if (!BuildOptions.bUseMikkTSpace)
 	{
-		bool bBuildSuccess = BuildSkeletalMesh_Legacy(LODModel, RefSkeleton, Influences, Wedges, Faces, Points, PointToOriginalMap, BuildOptions.OverlappingThresholds, BuildOptions.bComputeNormals, BuildOptions.bComputeTangents, OutWarningMessages, OutWarningNames);
+		bool bBuildSuccess = BuildSkeletalMesh_Legacy(LODModel, RefSkeleton, Influences, Wedges, Faces, Points, PointToOriginalMap, BuildOptions.OverlappingThresholds, BuildOptions.bComputeNormals, BuildOptions.bComputeTangents, BuildOptions.bComputeWeightedNormals, OutWarningMessages, OutWarningNames);
 		if (bBuildSuccess)
 		{
 			UpdateOverlappingVertices(LODModel);
@@ -4607,6 +4632,7 @@ bool FMeshUtilities::BuildSkeletalMesh_Legacy(FSkeletalMeshLODModel& LODModel
 											, const FOverlappingThresholds& OverlappingThresholds
 											, bool bComputeNormals
 											, bool bComputeTangents
+											, bool bComputeWeightedNormals
 											, TArray<FText> * OutWarningMessages
 											, TArray<FName> * OutWarningNames)
 {
@@ -4840,15 +4866,19 @@ bool FMeshUtilities::BuildSkeletalMesh_Legacy(FSkeletalMeshLODModel& LODModel
 				}
 			}
 
+			FVector FacePoint[3] = { Points[Wedges[Face.iWedge[0]].iVertex], Points[Wedges[Face.iWedge[1]].iVertex], Points[Wedges[Face.iWedge[2]].iVertex] };
 			// Process adjacent faces
 			for (int32 AdjacentFaceIndex = 0; AdjacentFaceIndex < AdjacentFaces.Num(); AdjacentFaceIndex++)
 			{
 				int32 OtherFaceIndex = AdjacentFaces[AdjacentFaceIndex];
 				const SkeletalMeshImportData::FMeshFace&	OtherFace = Faces[OtherFaceIndex];
+
+				FVector OtherFacePoint[3] = { Points[Wedges[OtherFace.iWedge[0]].iVertex], Points[Wedges[OtherFace.iWedge[1]].iVertex], Points[Wedges[OtherFace.iWedge[2]].iVertex] };
+				float OtherFaceArea = !bComputeWeightedNormals ? 1.0f : TriangleUtilities::ComputeTriangleArea(OtherFacePoint[0], OtherFacePoint[1], OtherFacePoint[2]);
 				FVector		OtherTriangleNormal = FPlane(
-					Points[Wedges[OtherFace.iWedge[2]].iVertex],
-					Points[Wedges[OtherFace.iWedge[1]].iVertex],
-					Points[Wedges[OtherFace.iWedge[0]].iVertex]
+					OtherFacePoint[2],
+					OtherFacePoint[1],
+					OtherFacePoint[0]
 					);
 				float		OtherFaceDeterminant = FVector::Triple(FaceTangentX[OtherFaceIndex], FaceTangentY[OtherFaceIndex], OtherTriangleNormal);
 
@@ -4857,22 +4887,27 @@ bool FMeshUtilities::BuildSkeletalMesh_Legacy(FSkeletalMeshLODModel& LODModel
 					for (int32 OtherVertexIndex = 0; OtherVertexIndex < 3; OtherVertexIndex++)
 					{
 						if (PointsEqual(
-							Points[Wedges[OtherFace.iWedge[OtherVertexIndex]].iVertex],
-							Points[Wedges[Face.iWedge[VertexIndex]].iVertex],
+							OtherFacePoint[OtherVertexIndex],
+							FacePoint[VertexIndex],
 							OverlappingThresholds
 							))
 						{
+							//Compute the angle
+							float OtherFaceAngle = !bComputeWeightedNormals ? 1.0f : TriangleUtilities::ComputeTriangleCornerAngle(OtherFacePoint[OtherVertexIndex], OtherFacePoint[(OtherVertexIndex + 1) % 3], OtherFacePoint[(OtherVertexIndex + 2) % 3]);
+							
+							float CornerWeight = (OtherFaceArea * OtherFaceAngle);
+
 							if (Determinant * OtherFaceDeterminant > 0.0f && SkeletalMeshTools::SkeletalMesh_UVsEqual(Wedges[OtherFace.iWedge[OtherVertexIndex]], Wedges[Face.iWedge[VertexIndex]], OverlappingThresholds))
 							{
-								VertexTangentX[VertexIndex] += FaceTangentX[OtherFaceIndex];
-								VertexTangentY[VertexIndex] += FaceTangentY[OtherFaceIndex];
+								VertexTangentX[VertexIndex] += CornerWeight *FaceTangentX[OtherFaceIndex];
+								VertexTangentY[VertexIndex] += CornerWeight *FaceTangentY[OtherFaceIndex];
 							}
 
 							// Only contribute 'normal' if the vertices are truly one and the same to obey hard "smoothing" edges baked into 
 							// the mesh by vertex duplication
 							if (Wedges[OtherFace.iWedge[OtherVertexIndex]].iVertex == Wedges[Face.iWedge[VertexIndex]].iVertex)
 							{
-								VertexTangentZ[VertexIndex] += OtherTriangleNormal;
+								VertexTangentZ[VertexIndex] += CornerWeight *OtherTriangleNormal;
 							}
 						}
 					}
diff --git a/Engine/Source/Developer/MeshUtilities/Private/MeshUtilitiesPrivate.h b/Engine/Source/Developer/MeshUtilities/Private/MeshUtilitiesPrivate.h
--- a/Engine/Source/Developer/MeshUtilities/Private/MeshUtilitiesPrivate.h
+++ b/Engine/Source/Developer/MeshUtilities/Private/MeshUtilitiesPrivate.h
@@ -109,7 +109,7 @@ private:
 	virtual bool GenerateUniqueUVsForStaticMesh(const FRawMesh& RawMesh, int32 TextureResolution, bool bMergeIdenticalMaterials, TArray<FVector2D>& OutTexCoords) const override;
 
 	virtual bool BuildSkeletalMesh(FSkeletalMeshLODModel& LODModel, const FReferenceSkeleton& RefSkeleton, const TArray<SkeletalMeshImportData::FVertInfluence>& Influences, const TArray<SkeletalMeshImportData::FMeshWedge>& Wedges, const TArray<SkeletalMeshImportData::FMeshFace>& Faces, const TArray<FVector>& Points, const TArray<int32>& PointToOriginalMap, const MeshBuildOptions& BuildOptions = MeshBuildOptions(), TArray<FText> * OutWarningMessages = NULL, TArray<FName> * OutWarningNames = NULL) override;
-	bool BuildSkeletalMesh_Legacy(FSkeletalMeshLODModel& LODModel, const FReferenceSkeleton& RefSkeleton, const TArray<SkeletalMeshImportData::FVertInfluence>& Influences, const TArray<SkeletalMeshImportData::FMeshWedge>& Wedges, const TArray<SkeletalMeshImportData::FMeshFace>& Faces, const TArray<FVector>& Points, const TArray<int32>& PointToOriginalMap, const FOverlappingThresholds& OverlappingThresholds, bool bComputeNormals = true, bool bComputeTangents = true, TArray<FText> * OutWarningMessages = NULL, TArray<FName> * OutWarningNames = NULL);
+	bool BuildSkeletalMesh_Legacy(FSkeletalMeshLODModel& LODModel, const FReferenceSkeleton& RefSkeleton, const TArray<SkeletalMeshImportData::FVertInfluence>& Influences, const TArray<SkeletalMeshImportData::FMeshWedge>& Wedges, const TArray<SkeletalMeshImportData::FMeshFace>& Faces, const TArray<FVector>& Points, const TArray<int32>& PointToOriginalMap, const FOverlappingThresholds& OverlappingThresholds, bool bComputeNormals = true, bool bComputeTangents = true, bool bComputeWeightedNormals = true, TArray<FText> * OutWarningMessages = NULL, TArray<FName> * OutWarningNames = NULL);
 
 	virtual void CacheOptimizeIndexBuffer(TArray<uint16>& Indices) override;
 	virtual void CacheOptimizeIndexBuffer(TArray<uint32>& Indices) override;
diff --git a/Engine/Source/Developer/MeshUtilities/Public/MeshUtilities.h b/Engine/Source/Developer/MeshUtilities/Public/MeshUtilities.h
--- a/Engine/Source/Developer/MeshUtilities/Public/MeshUtilities.h
+++ b/Engine/Source/Developer/MeshUtilities/Public/MeshUtilities.h
@@ -218,6 +218,7 @@ public:
 		, bComputeNormals(true)
 		, bComputeTangents(true)
 		, bUseMikkTSpace(false)
+		, bComputeWeightedNormals(false)
 		{
 		}
 
@@ -225,6 +226,7 @@ public:
 		bool bComputeNormals;
 		bool bComputeTangents;
 		bool bUseMikkTSpace;
+		bool bComputeWeightedNormals;
 		FOverlappingThresholds OverlappingThresholds;
 	};
 	
diff --git a/Engine/Source/Developer/MeshUtilitiesCommon/Public/MeshUtilitiesCommon.h b/Engine/Source/Developer/MeshUtilitiesCommon/Public/MeshUtilitiesCommon.h
--- a/Engine/Source/Developer/MeshUtilitiesCommon/Public/MeshUtilitiesCommon.h
+++ b/Engine/Source/Developer/MeshUtilitiesCommon/Public/MeshUtilitiesCommon.h
@@ -53,3 +53,31 @@ inline bool PointsEqual(const FVector& V1, const FVector& V2, float ComparisonTh
 	}
 	return true;
 }
+
+namespace TriangleUtilities
+{
+	/*
+	 * This function compute the area of a triangle, it will return zero if the triangle is degenerated
+	 */
+	static float ComputeTriangleArea(const FVector& PointA, const FVector& PointB, const FVector& PointC)
+	{
+		return FVector::CrossProduct((PointB - PointA), (PointC - PointA)).Size() / 2.0f;
+	}
+
+	/*
+	 * This function compute the angle of a triangle corner, it will return zero if the triangle is degenerated
+	 */
+	static float ComputeTriangleCornerAngle(const FVector& PointA, const FVector& PointB, const FVector& PointC)
+	{
+		FVector E1 = (PointB - PointA);
+		FVector E2 = (PointC - PointA);
+		//Normalize both edges (unit vector) of the triangle so we get a dotProduct result that will be a valid acos input [-1, 1]
+		if (!E1.Normalize() || !E2.Normalize())
+		{
+			//Return a null ratio if the polygon is degenerate
+			return 0.0f;
+		}
+		float DotProduct = FVector::DotProduct(E1, E2);
+		return FMath::Acos(DotProduct);
+	}
+}
\ No newline at end of file
diff --git a/Engine/Source/Developer/SkeletalMeshUtilitiesCommon/Private/LODUtilities.cpp b/Engine/Source/Developer/SkeletalMeshUtilitiesCommon/Private/LODUtilities.cpp
--- a/Engine/Source/Developer/SkeletalMeshUtilitiesCommon/Private/LODUtilities.cpp
+++ b/Engine/Source/Developer/SkeletalMeshUtilitiesCommon/Private/LODUtilities.cpp
@@ -1598,7 +1598,7 @@ void MatchVertexIndexUsingPosition(
 	}
 }
 
-bool FLODUtilities::UpdateAlternateSkinWeights(USkeletalMesh* SkeletalMeshDest, const FName& ProfileNameDest, int32 LODIndexDest, FOverlappingThresholds OverlappingThresholds, bool ShouldImportNormals, bool ShouldImportTangents, bool bUseMikkTSpace)
+bool FLODUtilities::UpdateAlternateSkinWeights(USkeletalMesh* SkeletalMeshDest, const FName& ProfileNameDest, int32 LODIndexDest, FOverlappingThresholds OverlappingThresholds, bool ShouldImportNormals, bool ShouldImportTangents, bool bUseMikkTSpace, bool bComputeWeightedNormals)
 {
 	//Ensure log message only once
 	bool bNoMatchMsgDone = false;
@@ -1871,6 +1871,7 @@ bool FLODUtilities::UpdateAlternateSkinWeights(USkeletalMesh* SkeletalMeshDest,
 	BuildOptions.bComputeNormals = !ShouldImportNormals || !ImportDataDest.bHasNormals;
 	BuildOptions.bComputeTangents = !ShouldImportTangents || !ImportDataDest.bHasTangents;
 	BuildOptions.bUseMikkTSpace = (bUseMikkTSpace) && (!ShouldImportNormals || !ShouldImportTangents);
+	BuildOptions.bComputeWeightedNormals = bComputeWeightedNormals;
 	BuildOptions.bRemoveDegenerateTriangles = false;
 
 	//Build the skeletal mesh asset
@@ -1884,7 +1885,7 @@ bool FLODUtilities::UpdateAlternateSkinWeights(USkeletalMesh* SkeletalMeshDest,
 	return bBuildSuccess;
 }
 
-bool FLODUtilities::UpdateAlternateSkinWeights(USkeletalMesh* SkeletalMeshDest, const FName& ProfileNameDest, USkeletalMesh* SkeletalMeshSrc, int32 LODIndexDest, int32 LODIndexSrc, FOverlappingThresholds OverlappingThresholds, bool ShouldImportNormals, bool ShouldImportTangents, bool bUseMikkTSpace)
+bool FLODUtilities::UpdateAlternateSkinWeights(USkeletalMesh* SkeletalMeshDest, const FName& ProfileNameDest, USkeletalMesh* SkeletalMeshSrc, int32 LODIndexDest, int32 LODIndexSrc, FOverlappingThresholds OverlappingThresholds, bool ShouldImportNormals, bool ShouldImportTangents, bool bUseMikkTSpace, bool bComputeWeightedNormals)
 {
 	//Grab all the destination structure
 	check(SkeletalMeshDest);
@@ -1928,7 +1929,7 @@ bool FLODUtilities::UpdateAlternateSkinWeights(USkeletalMesh* SkeletalMeshDest,
 	LODModelDest.RawSkeletalMeshBulkData.SaveRawMesh(ImportDataDest);
 
 	//Build the alternate buffer with all the data into the bulk
-	return UpdateAlternateSkinWeights(SkeletalMeshDest, ProfileNameDest, LODIndexDest, OverlappingThresholds, ShouldImportNormals, ShouldImportTangents, bUseMikkTSpace);
+	return UpdateAlternateSkinWeights(SkeletalMeshDest, ProfileNameDest, LODIndexDest, OverlappingThresholds, ShouldImportNormals, ShouldImportTangents, bUseMikkTSpace, bComputeWeightedNormals);
 }
 
 void FLODUtilities::GenerateImportedSkinWeightProfileData(const FSkeletalMeshLODModel& LODModelDest, FImportedSkinWeightProfileData &ImportedProfileData)
diff --git a/Engine/Source/Developer/SkeletalMeshUtilitiesCommon/Public/LODUtilities.h b/Engine/Source/Developer/SkeletalMeshUtilitiesCommon/Public/LODUtilities.h
--- a/Engine/Source/Developer/SkeletalMeshUtilitiesCommon/Public/LODUtilities.h
+++ b/Engine/Source/Developer/SkeletalMeshUtilitiesCommon/Public/LODUtilities.h
@@ -105,7 +105,7 @@ public:
 	 * @param LODIndexDest - the destination LOD
 	 * @param LODIndexSrc - the Source LOD index
 	 */
-	static bool UpdateAlternateSkinWeights(USkeletalMesh* SkeletalMeshDest, const FName& ProfileNameDest, USkeletalMesh* SkeletalMeshSrc, int32 LODIndexDest, int32 LODIndexSrc, FOverlappingThresholds OverlappingThresholds, bool ShouldImportNormals, bool ShouldImportTangents, bool bUseMikkTSpace);
+	static bool UpdateAlternateSkinWeights(USkeletalMesh* SkeletalMeshDest, const FName& ProfileNameDest, USkeletalMesh* SkeletalMeshSrc, int32 LODIndexDest, int32 LODIndexSrc, FOverlappingThresholds OverlappingThresholds, bool ShouldImportNormals, bool ShouldImportTangents, bool bUseMikkTSpace, bool bComputeWeightedNormals);
 
 	
 	/*
@@ -115,7 +115,7 @@ public:
 	 * @param SkeletalMeshDest - The skeletal mesh that will receive the alternate skinning weights.
 	 * @param LODIndexDest - the destination LOD
 	 */
-	static bool UpdateAlternateSkinWeights(USkeletalMesh* SkeletalMeshDest, const FName& ProfileNameDest, int32 LODIndexDest, FOverlappingThresholds OverlappingThresholds, bool ShouldImportNormals, bool ShouldImportTangents, bool bUseMikkTSpace);
+	static bool UpdateAlternateSkinWeights(USkeletalMesh* SkeletalMeshDest, const FName& ProfileNameDest, int32 LODIndexDest, FOverlappingThresholds OverlappingThresholds, bool ShouldImportNormals, bool ShouldImportTangents, bool bUseMikkTSpace, bool bComputeWeightedNormals);
 
 	/** Re-generate all (editor-only) skin weight profile, used whenever we rebuild the skeletal mesh data which could change the chunking and bone indices */
 	static void RegenerateAllImportSkinWeightProfileData(FSkeletalMeshLODModel& LODModelDest);
diff --git a/Engine/Source/Editor/StaticMeshEditor/Private/StaticMeshEditorTools.cpp b/Engine/Source/Editor/StaticMeshEditor/Private/StaticMeshEditorTools.cpp
--- a/Engine/Source/Editor/StaticMeshEditor/Private/StaticMeshEditorTools.cpp
+++ b/Engine/Source/Editor/StaticMeshEditor/Private/StaticMeshEditorTools.cpp
@@ -504,6 +504,22 @@ void FMeshBuildSettingsLayout::GenerateChildContent( IDetailChildrenBuilder& Chi
 		];
 	}
 
+	{
+		ChildrenBuilder.AddCustomRow( LOCTEXT("ComputeWeightedNormals", "Compute Weighted Normals") )
+		.NameContent()
+		[
+			SNew(STextBlock)
+			.Font( IDetailLayoutBuilder::GetDetailFont() )
+			.Text(LOCTEXT("ComputeWeightedNormals", "Compute Weighted Normals"))
+		]
+		.ValueContent()
+		[
+			SNew(SCheckBox)
+			.IsChecked(this, &FMeshBuildSettingsLayout::ShouldComputeWeightedNormals)
+			.OnCheckStateChanged(this, &FMeshBuildSettingsLayout::OnComputeWeightedNormalsChanged)
+		];
+	}
+
 	{
 		ChildrenBuilder.AddCustomRow( LOCTEXT("RemoveDegenerates", "Remove Degenerates") )
 		.NameContent()
@@ -786,6 +802,11 @@ ECheckBoxState FMeshBuildSettingsLayout::ShouldUseMikkTSpace() const
 	return BuildSettings.bUseMikkTSpace ? ECheckBoxState::Checked : ECheckBoxState::Unchecked;
 }
 
+ECheckBoxState FMeshBuildSettingsLayout::ShouldComputeWeightedNormals() const
+{
+	return BuildSettings.bComputeWeightedNormals ? ECheckBoxState::Checked : ECheckBoxState::Unchecked;
+}
+
 ECheckBoxState FMeshBuildSettingsLayout::ShouldRemoveDegenerates() const
 {
 	return BuildSettings.bRemoveDegenerates ? ECheckBoxState::Checked : ECheckBoxState::Unchecked;
@@ -891,6 +912,19 @@ void FMeshBuildSettingsLayout::OnUseMikkTSpaceChanged(ECheckBoxState NewState)
 	}
 }
 
+void FMeshBuildSettingsLayout::OnComputeWeightedNormalsChanged(ECheckBoxState NewState)
+{
+	const bool bComputeWeightedNormals = (NewState == ECheckBoxState::Checked) ? true : false;
+	if (BuildSettings.bComputeWeightedNormals != bComputeWeightedNormals)
+	{
+		if (FEngineAnalytics::IsAvailable())
+		{
+			FEngineAnalytics::GetProvider().RecordEvent(TEXT("Editor.Usage.StaticMesh.BuildSettings"), TEXT("bComputeWeightedNormals"), bComputeWeightedNormals ? TEXT("True") : TEXT("False"));
+		}
+		BuildSettings.bComputeWeightedNormals = bComputeWeightedNormals;
+	}
+}
+
 void FMeshBuildSettingsLayout::OnRemoveDegeneratesChanged(ECheckBoxState NewState)
 {
 	const bool bRemoveDegenerates = (NewState == ECheckBoxState::Checked) ? true : false;
diff --git a/Engine/Source/Editor/StaticMeshEditor/Private/StaticMeshEditorTools.h b/Engine/Source/Editor/StaticMeshEditor/Private/StaticMeshEditorTools.h
--- a/Engine/Source/Editor/StaticMeshEditor/Private/StaticMeshEditorTools.h
+++ b/Engine/Source/Editor/StaticMeshEditor/Private/StaticMeshEditorTools.h
@@ -183,6 +183,7 @@ private:
 	ECheckBoxState ShouldRecomputeNormals() const;
 	ECheckBoxState ShouldRecomputeTangents() const;
 	ECheckBoxState ShouldUseMikkTSpace() const;
+	ECheckBoxState ShouldComputeWeightedNormals() const;
 	ECheckBoxState ShouldRemoveDegenerates() const;
 	ECheckBoxState ShouldBuildAdjacencyBuffer() const;
 	ECheckBoxState ShouldBuildReversedIndexBuffer() const;
@@ -201,6 +202,7 @@ private:
 	void OnRecomputeNormalsChanged(ECheckBoxState NewState);
 	void OnRecomputeTangentsChanged(ECheckBoxState NewState);
 	void OnUseMikkTSpaceChanged(ECheckBoxState NewState);
+	void OnComputeWeightedNormalsChanged(ECheckBoxState NewState);
 	void OnRemoveDegeneratesChanged(ECheckBoxState NewState);
 	void OnBuildAdjacencyBufferChanged(ECheckBoxState NewState);
 	void OnBuildReversedIndexBufferChanged(ECheckBoxState NewState);
diff --git a/Engine/Source/Editor/UnrealEd/Classes/Factories/FbxMeshImportData.h b/Engine/Source/Editor/UnrealEd/Classes/Factories/FbxMeshImportData.h
--- a/Engine/Source/Editor/UnrealEd/Classes/Factories/FbxMeshImportData.h
+++ b/Engine/Source/Editor/UnrealEd/Classes/Factories/FbxMeshImportData.h
@@ -82,6 +82,10 @@ class UFbxMeshImportData : public UFbxAssetImportData
 	UPROPERTY(EditAnywhere, BlueprintReadWrite, config, AdvancedDisplay, Category = Mesh, meta=(ImportType="Mesh|GeoOnly"))
 	TEnumAsByte<enum EFBXNormalGenerationMethod::Type> NormalGenerationMethod;
 
+	/** Enabling this option will use weighted normals algorithm (area and angle) when computing normals or tangents */
+	UPROPERTY(EditAnywhere, BlueprintReadWrite, AdvancedDisplay, config, Category = Mesh, meta = (OBJRestrict = "true", ImportType = "Mesh|GeoOnly"))
+	uint32 bComputeWeightedNormals : 1;
+
 	/* If checked, The material list will be reorder to the same order has the FBX file. */
 	UPROPERTY(EditAnywhere, BlueprintReadWrite, config, AdvancedDisplay, Category = Material, meta = (OBJRestrict = "true"))
 	bool bReorderMaterialToFbxOrder;
diff --git a/Engine/Source/Editor/UnrealEd/Classes/Factories/FbxSkeletalMeshImportData.h b/Engine/Source/Editor/UnrealEd/Classes/Factories/FbxSkeletalMeshImportData.h
--- a/Engine/Source/Editor/UnrealEd/Classes/Factories/FbxSkeletalMeshImportData.h
+++ b/Engine/Source/Editor/UnrealEd/Classes/Factories/FbxSkeletalMeshImportData.h
@@ -29,6 +29,8 @@ class UFbxSkeletalMeshImportData : public UFbxMeshImportData
 {
 	GENERATED_UCLASS_BODY()
 public:
+	virtual void Serialize(FArchive& Ar) override;
+
 	/** Filter the content we want to import from the incoming FBX skeletal mesh.*/
 	UPROPERTY(EditAnywhere, Category = Mesh, meta = (ImportType = "SkeletalMesh", DisplayName = "Import Content Type", OBJRestrict = "true"))
 	TEnumAsByte<enum EFBXImportContentType> ImportContentType;
diff --git a/Engine/Source/Editor/UnrealEd/Private/Fbx/FbxMainImport.cpp b/Engine/Source/Editor/UnrealEd/Private/Fbx/FbxMainImport.cpp
--- a/Engine/Source/Editor/UnrealEd/Private/Fbx/FbxMainImport.cpp
+++ b/Engine/Source/Editor/UnrealEd/Private/Fbx/FbxMainImport.cpp
@@ -389,6 +389,7 @@ void ApplyImportUIToImportOptions(UFbxImportUI* ImportUI, FBXImportOptions& InOu
 		UFbxStaticMeshImportData* StaticMeshData	= ImportUI->StaticMeshImportData;
 		InOutImportOptions.NormalImportMethod		= StaticMeshData->NormalImportMethod;
 		InOutImportOptions.NormalGenerationMethod	= StaticMeshData->NormalGenerationMethod;
+		InOutImportOptions.bComputeWeightedNormals	= StaticMeshData->bComputeWeightedNormals;
 		InOutImportOptions.ImportTranslation		= StaticMeshData->ImportTranslation;
 		InOutImportOptions.ImportRotation			= StaticMeshData->ImportRotation;
 		InOutImportOptions.ImportUniformScale		= StaticMeshData->ImportUniformScale;
@@ -409,6 +410,7 @@ void ApplyImportUIToImportOptions(UFbxImportUI* ImportUI, FBXImportOptions& InOu
 		InOutImportOptions.bImportAsSkeletalSkinning	= SkeletalMeshData->ImportContentType == EFBXImportContentType::FBXICT_SkinningWeights;
 		InOutImportOptions.NormalImportMethod			= SkeletalMeshData->NormalImportMethod;
 		InOutImportOptions.NormalGenerationMethod		= SkeletalMeshData->NormalGenerationMethod;
+		InOutImportOptions.bComputeWeightedNormals		= SkeletalMeshData->bComputeWeightedNormals;
 		InOutImportOptions.ImportTranslation			= SkeletalMeshData->ImportTranslation;
 		InOutImportOptions.ImportRotation				= SkeletalMeshData->ImportRotation;
 		InOutImportOptions.ImportUniformScale			= SkeletalMeshData->ImportUniformScale;
@@ -438,6 +440,7 @@ void ApplyImportUIToImportOptions(UFbxImportUI* ImportUI, FBXImportOptions& InOu
 	{
 		UFbxAnimSequenceImportData* AnimData	= ImportUI->AnimSequenceImportData;
 		InOutImportOptions.NormalImportMethod = FBXNIM_ComputeNormals;
+		InOutImportOptions.bComputeWeightedNormals = true;
 		InOutImportOptions.ImportTranslation	= AnimData->ImportTranslation;
 		InOutImportOptions.ImportRotation		= AnimData->ImportRotation;
 		InOutImportOptions.ImportUniformScale	= AnimData->ImportUniformScale;
@@ -1526,6 +1529,7 @@ bool FFbxImporter::ImportFromFile(const FString& Filename, const FString& Type,
 							Attribs.Add(FAnalyticsEventAttribute(TEXT("MeshOpt ImportRigidMesh"), CaptureImportOptions->bImportRigidMesh));
 							Attribs.Add(FAnalyticsEventAttribute(TEXT("MeshOpt NormalGenerationMethod"), FBXNormalGenerationMethodEnum->GetNameStringByValue(CaptureImportOptions->NormalGenerationMethod)));
 							Attribs.Add(FAnalyticsEventAttribute(TEXT("MeshOpt NormalImportMethod"), FBXNormalImportMethodEnum->GetNameStringByValue(CaptureImportOptions->NormalImportMethod)));
+							Attribs.Add(FAnalyticsEventAttribute(TEXT("MeshOpt ComputeWeightedNormals"), CaptureImportOptions->bComputeWeightedNormals));
 						};
 
 						auto AddSKAnalytic = [&Attribs, &CaptureImportOptions]()
diff --git a/Engine/Source/Editor/UnrealEd/Private/Fbx/FbxMeshImportData.cpp b/Engine/Source/Editor/UnrealEd/Private/Fbx/FbxMeshImportData.cpp
--- a/Engine/Source/Editor/UnrealEd/Private/Fbx/FbxMeshImportData.cpp
+++ b/Engine/Source/Editor/UnrealEd/Private/Fbx/FbxMeshImportData.cpp
@@ -8,6 +8,7 @@ UFbxMeshImportData::UFbxMeshImportData(const FObjectInitializer& ObjectInitializ
 {
 	NormalImportMethod = FBXNIM_ComputeNormals;
 	NormalGenerationMethod = EFBXNormalGenerationMethod::MikkTSpace;
+	bComputeWeightedNormals = true;
 	bBakePivotInVertex = false;
 	bReorderMaterialToFbxOrder = true;
 }
@@ -23,7 +24,8 @@ bool UFbxMeshImportData::CanEditChange(const UProperty* InProperty) const
 	}
 
 	static FName NAME_NormalGenerationMethod("NormalGenerationMethod");
-	if( bMutable && InProperty->GetFName() == NAME_NormalGenerationMethod )
+	static FName NAME_ComputeWeightedNormals("bComputeWeightedNormals");
+	if( bMutable && (InProperty->GetFName() == NAME_NormalGenerationMethod || InProperty->GetFName() == NAME_ComputeWeightedNormals))
 	{
 		// Normal generation method is ignored if we are importing both tangents and normals
 		return NormalImportMethod == FBXNIM_ImportNormalsAndTangents ? false : true;
diff --git a/Engine/Source/Editor/UnrealEd/Private/Fbx/FbxSkeletalMeshImport.cpp b/Engine/Source/Editor/UnrealEd/Private/Fbx/FbxSkeletalMeshImport.cpp
--- a/Engine/Source/Editor/UnrealEd/Private/Fbx/FbxSkeletalMeshImport.cpp
+++ b/Engine/Source/Editor/UnrealEd/Private/Fbx/FbxSkeletalMeshImport.cpp
@@ -1891,6 +1891,7 @@ USkeletalMesh* UnFbx::FFbxImporter::ImportSkeletalMesh(FImportSkeletalMeshArgs &
 		BuildOptions.bComputeNormals = !ImportOptions->ShouldImportNormals() || !SkelMeshImportDataPtr->bHasNormals;
 		BuildOptions.bComputeTangents = !ImportOptions->ShouldImportTangents() || !SkelMeshImportDataPtr->bHasTangents;
 		BuildOptions.bUseMikkTSpace = (ImportOptions->NormalGenerationMethod == EFBXNormalGenerationMethod::MikkTSpace) && (!ImportOptions->ShouldImportNormals() || !ImportOptions->ShouldImportTangents());
+		BuildOptions.bComputeWeightedNormals = ImportOptions->bComputeWeightedNormals;
 		BuildOptions.bRemoveDegenerateTriangles = false;
 
 		IMeshUtilities& MeshUtilities = FModuleManager::Get().LoadModuleChecked<IMeshUtilities>("MeshUtilities");
diff --git a/Engine/Source/Editor/UnrealEd/Private/Fbx/FbxSkeletalMeshImportData.cpp b/Engine/Source/Editor/UnrealEd/Private/Fbx/FbxSkeletalMeshImportData.cpp
--- a/Engine/Source/Editor/UnrealEd/Private/Fbx/FbxSkeletalMeshImportData.cpp
+++ b/Engine/Source/Editor/UnrealEd/Private/Fbx/FbxSkeletalMeshImportData.cpp
@@ -3,6 +3,7 @@
 #include "Factories/FbxSkeletalMeshImportData.h"
 #include "Engine/SkeletalMesh.h"
 #include "UObject/Package.h"
+#include "UObject/EditorObjectVersion.h"
 
 UFbxSkeletalMeshImportData::UFbxSkeletalMeshImportData(const FObjectInitializer& ObjectInitializer)
 	: Super(ObjectInitializer)
@@ -14,6 +15,19 @@ UFbxSkeletalMeshImportData::UFbxSkeletalMeshImportData(const FObjectInitializer&
 	LastImportContentType = EFBXImportContentType::FBXICT_All;
 }
 
+void UFbxSkeletalMeshImportData::Serialize(FArchive& Ar)
+{
+	Super::Serialize(Ar);
+	Ar.UsingCustomVersion(FEditorObjectVersion::GUID);
+
+	//bComputeWeightedNormals default is true so old asset should have set to false.
+	if (Ar.CustomVer(FEditorObjectVersion::GUID) < FEditorObjectVersion::ComputeWeightedNormals)
+	{
+		//Old asset are not using weighted normals
+		bComputeWeightedNormals = false;
+	}
+}
+
 UFbxSkeletalMeshImportData* UFbxSkeletalMeshImportData::GetImportDataForSkeletalMesh(USkeletalMesh* SkeletalMesh, UFbxSkeletalMeshImportData* TemplateForCreation)
 {
 	check(SkeletalMesh);
diff --git a/Engine/Source/Editor/UnrealEd/Private/Fbx/FbxStaticMeshImport.cpp b/Engine/Source/Editor/UnrealEd/Private/Fbx/FbxStaticMeshImport.cpp
--- a/Engine/Source/Editor/UnrealEd/Private/Fbx/FbxStaticMeshImport.cpp
+++ b/Engine/Source/Editor/UnrealEd/Private/Fbx/FbxStaticMeshImport.cpp
@@ -1690,6 +1690,7 @@ UStaticMesh* UnFbx::FFbxImporter::ImportStaticMeshAsSingle(UObject* InParent, TA
 		SrcModel.BuildSettings.bRecomputeNormals = ImportOptions->NormalImportMethod == FBXNIM_ComputeNormals;
 		SrcModel.BuildSettings.bRecomputeTangents = ImportOptions->NormalImportMethod != FBXNIM_ImportNormalsAndTangents;
 		SrcModel.BuildSettings.bUseMikkTSpace = (ImportOptions->NormalGenerationMethod == EFBXNormalGenerationMethod::MikkTSpace) && (!ImportOptions->ShouldImportNormals() || !ImportOptions->ShouldImportTangents());
+		SrcModel.BuildSettings.bComputeWeightedNormals = ImportOptions->bComputeWeightedNormals;
 		if (ImportOptions->bGenerateLightmapUVs)
 		{
 			SrcModel.BuildSettings.bGenerateLightmapUVs = true;
diff --git a/Engine/Source/Editor/UnrealEd/Private/SkinWeightsUtilities.cpp b/Engine/Source/Editor/UnrealEd/Private/SkinWeightsUtilities.cpp
--- a/Engine/Source/Editor/UnrealEd/Private/SkinWeightsUtilities.cpp
+++ b/Engine/Source/Editor/UnrealEd/Private/SkinWeightsUtilities.cpp
@@ -206,14 +206,14 @@ bool FSkinWeightsUtilities::ImportAlternateSkinWeight(USkeletalMesh* SkeletalMes
 					SkeletalMesh->ReleaseResources();
 					SkeletalMesh->ReleaseResourcesFence.Wait();
 
-					bResult = FLODUtilities::UpdateAlternateSkinWeights(SkeletalMesh, ProfileName, TmpSkeletalMesh, TargetLODIndex, SrcLodIndex, OverlappingThresholds, ShouldImportNormals, ShouldImportTangents, bUseMikkTSpace);
+					bResult = FLODUtilities::UpdateAlternateSkinWeights(SkeletalMesh, ProfileName, TmpSkeletalMesh, TargetLODIndex, SrcLodIndex, OverlappingThresholds, ShouldImportNormals, ShouldImportTangents, bUseMikkTSpace, ImportOptions.bComputeWeightedNormals);
 
 					SkeletalMesh->PostEditChange();
 					SkeletalMesh->InitResources();
 				}
 				else
 				{
-					bResult = FLODUtilities::UpdateAlternateSkinWeights(SkeletalMesh, ProfileName, TmpSkeletalMesh, TargetLODIndex, SrcLodIndex, OverlappingThresholds, ShouldImportNormals, ShouldImportTangents, bUseMikkTSpace);
+					bResult = FLODUtilities::UpdateAlternateSkinWeights(SkeletalMesh, ProfileName, TmpSkeletalMesh, TargetLODIndex, SrcLodIndex, OverlappingThresholds, ShouldImportNormals, ShouldImportTangents, bUseMikkTSpace, ImportOptions.bComputeWeightedNormals);
 				}
 				
 				if (!bResult)
@@ -343,6 +343,7 @@ bool FSkinWeightsUtilities::RemoveSkinnedWeightProfileData(USkeletalMesh* Skelet
 	BuildOptions.bComputeNormals = !bShouldImportNormals || !ImportDataDest.bHasNormals;
 	BuildOptions.bComputeTangents = !bShouldImportTangents || !ImportDataDest.bHasTangents;
 	BuildOptions.bUseMikkTSpace = (OriginalSkeletalMeshImportData->NormalGenerationMethod == EFBXNormalGenerationMethod::MikkTSpace) && (!bShouldImportNormals || !bShouldImportTangents);
+	BuildOptions.bComputeWeightedNormals = OriginalSkeletalMeshImportData->bComputeWeightedNormals;
 	BuildOptions.bRemoveDegenerateTriangles = false;
 
 	//Build the skeletal mesh asset
diff --git a/Engine/Source/Editor/UnrealEd/Public/FbxImporter.h b/Engine/Source/Editor/UnrealEd/Public/FbxImporter.h
--- a/Engine/Source/Editor/UnrealEd/Public/FbxImporter.h
+++ b/Engine/Source/Editor/UnrealEd/Public/FbxImporter.h
@@ -138,6 +138,7 @@ struct FBXImportOptions
 	float ImportUniformScale;
 	EFBXNormalImportMethod NormalImportMethod;
 	EFBXNormalGenerationMethod::Type NormalGenerationMethod;
+	bool bComputeWeightedNormals;
 	bool bTransformVertexToAbsolute;
 	bool bBakePivotInVertex;
 	EFBXImportType ImportType;
diff --git a/Engine/Source/Runtime/Core/Public/UObject/EditorObjectVersion.h b/Engine/Source/Runtime/Core/Public/UObject/EditorObjectVersion.h
--- a/Engine/Source/Runtime/Core/Public/UObject/EditorObjectVersion.h
+++ b/Engine/Source/Runtime/Core/Public/UObject/EditorObjectVersion.h
@@ -81,6 +81,8 @@ struct CORE_API FEditorObjectVersion
 		RemoveLandscapeHoleMaterial,
 		// MeshDescription defined by triangles instead of arbitrary polygons
 		MeshDescriptionTriangles,
+		//Add weighted area and angle when computing the normals
+		ComputeWeightedNormals,
 		// -----<new versions can be added above this line>-------------------------------------------------
 		VersionPlusOne,
 		LatestVersion = VersionPlusOne - 1
diff --git a/Engine/Source/Runtime/Engine/Classes/Engine/EngineTypes.h b/Engine/Source/Runtime/Engine/Classes/Engine/EngineTypes.h
--- a/Engine/Source/Runtime/Engine/Classes/Engine/EngineTypes.h
+++ b/Engine/Source/Runtime/Engine/Classes/Engine/EngineTypes.h
@@ -2555,6 +2555,10 @@ struct FMeshBuildSettings
 	UPROPERTY(EditAnywhere, Category=BuildSettings)
 	uint8 bRecomputeTangents:1;
 
+	/** If true, we will use the surface area and the corner angle of the triangle as a ratio when computing the normals. */
+	UPROPERTY(EditAnywhere, Category = BuildSettings)
+	uint8 bComputeWeightedNormals : 1;
+
 	/** If true, degenerate triangles will be removed. */
 	UPROPERTY(EditAnywhere, Category=BuildSettings)
 	uint8 bRemoveDegenerates:1;
@@ -2621,6 +2625,7 @@ struct FMeshBuildSettings
 		: bUseMikkTSpace(true)
 		, bRecomputeNormals(true)
 		, bRecomputeTangents(true)
+		, bComputeWeightedNormals(false)
 		, bRemoveDegenerates(true)
 		, bBuildAdjacencyBuffer(true)
 		, bBuildReversedIndexBuffer(true)
@@ -2645,6 +2650,7 @@ struct FMeshBuildSettings
 	{
 		return bRecomputeNormals == Other.bRecomputeNormals
 			&& bRecomputeTangents == Other.bRecomputeTangents
+			&& bComputeWeightedNormals == Other.bComputeWeightedNormals
 			&& bUseMikkTSpace == Other.bUseMikkTSpace
 			&& bRemoveDegenerates == Other.bRemoveDegenerates
 			&& bBuildAdjacencyBuffer == Other.bBuildAdjacencyBuffer
diff --git a/Engine/Source/Runtime/Engine/Private/StaticMesh.cpp b/Engine/Source/Runtime/Engine/Private/StaticMesh.cpp
--- a/Engine/Source/Runtime/Engine/Private/StaticMesh.cpp
+++ b/Engine/Source/Runtime/Engine/Private/StaticMesh.cpp
@@ -1981,6 +1981,7 @@ static void SerializeBuildSettingsForDDC(FArchive& Ar, FMeshBuildSettings& Build
 	FArchive_Serialize_BitfieldBool(Ar, BuildSettings.bRecomputeNormals);
 	FArchive_Serialize_BitfieldBool(Ar, BuildSettings.bRecomputeTangents);
 	FArchive_Serialize_BitfieldBool(Ar, BuildSettings.bUseMikkTSpace);
+	FArchive_Serialize_BitfieldBool(Ar, BuildSettings.bComputeWeightedNormals);
 	FArchive_Serialize_BitfieldBool(Ar, BuildSettings.bRemoveDegenerates);
 	FArchive_Serialize_BitfieldBool(Ar, BuildSettings.bBuildAdjacencyBuffer);
 	FArchive_Serialize_BitfieldBool(Ar, BuildSettings.bBuildReversedIndexBuffer);
