Fixed deprecation warning.

diff --git a/Engine/Source/Developer/MeshDescriptionOperations/Private/MeshDescriptionOperations.cpp b/Engine/Source/Developer/MeshDescriptionOperations/Private/MeshDescriptionOperations.cpp
--- a/Engine/Source/Developer/MeshDescriptionOperations/Private/MeshDescriptionOperations.cpp
+++ b/Engine/Source/Developer/MeshDescriptionOperations/Private/MeshDescriptionOperations.cpp
@@ -67,7 +67,7 @@ void FMeshDescriptionOperations::RecomputeNormalsAndTangentsIfNeeded(FMeshDescri
 
 void FMeshDescriptionOperations::CreatePolygonNTB(FMeshDescription& MeshDescription, float ComparisonThreshold)
 {
-	FStaticMeshOperations::ComputePolygonTangentsAndNormals(MeshDescription, ComparisonThreshold);
+	FStaticMeshOperations::ComputeTriangleTangentsAndNormals(MeshDescription, ComparisonThreshold);
 }
 
 void FMeshDescriptionOperations::CreateNormals(FMeshDescription& MeshDescription, FMeshDescriptionOperations::ETangentOptions TangentOptions, bool bComputeTangent)
