CIS Fix
#rnx

diff --git a/Engine/Source/Developer/MeshDescriptionOperations/Private/MeshDescriptionOperations.cpp b/Engine/Source/Developer/MeshDescriptionOperations/Private/MeshDescriptionOperations.cpp
--- a/Engine/Source/Developer/MeshDescriptionOperations/Private/MeshDescriptionOperations.cpp
+++ b/Engine/Source/Developer/MeshDescriptionOperations/Private/MeshDescriptionOperations.cpp
@@ -110,7 +110,12 @@ bool FMeshDescriptionOperations::GenerateUniqueUVsForStaticMesh(const FMeshDescr
 {
 	TRACE_CPUPROFILER_EVENT_SCOPE(FMeshDescriptionOperations::GenerateUniqueUVsForStaticMesh);
 
-	return FStaticMeshOperations::GenerateUniqueUVsForStaticMesh(MeshDescription, TextureResolution, bMergeIdenticalMaterials, OutTexCoords);
+	FStaticMeshOperations::FGenerateUVOptions GenerateUVOptions;
+	GenerateUVOptions.TextureResolution = TextureResolution;
+	GenerateUVOptions.bMergeTrianglesWithIdenticalAttributes = bMergeIdenticalMaterials;
+	GenerateUVOptions.UVMethod = FStaticMeshOperations::EGenerateUVMethod::Legacy;
+
+	return FStaticMeshOperations::GenerateUV(MeshDescription, GenerateUVOptions, OutTexCoords);
 }
 
 bool FMeshDescriptionOperations::AddUVChannel(FMeshDescription& MeshDescription)
